Fix recast module compile issues

diff --git a/Engine/Source/Runtime/Navmesh/Private/DebugUtils/RecastDump.cpp b/Engine/Source/Runtime/Navmesh/Private/DebugUtils/RecastDump.cpp
--- a/Engine/Source/Runtime/Navmesh/Private/DebugUtils/RecastDump.cpp
+++ b/Engine/Source/Runtime/Navmesh/Private/DebugUtils/RecastDump.cpp
@@ -37,7 +37,7 @@ static void ioprintf(duFileIO* io, const char* format, ...)
 	char line[256];
 	va_list ap;
 	va_start(ap, format);
-	const int n = vsnprintf(line, sizeof(line), format, ap);
+	const int n = FCStringAnsi::GetVarArgs(line, ARRAY_COUNT(line), ARRAY_COUNT(line) - 1, format, ap);
 	va_end(ap);
 	if (n > 0)
 		io->write(line, sizeof(char)*n);
diff --git a/Engine/Source/Runtime/Navmesh/Private/NavmeshModulePrivatePCH.h b/Engine/Source/Runtime/Navmesh/Private/NavmeshModulePrivatePCH.h
--- a/Engine/Source/Runtime/Navmesh/Private/NavmeshModulePrivatePCH.h
+++ b/Engine/Source/Runtime/Navmesh/Private/NavmeshModulePrivatePCH.h
@@ -2,4 +2,4 @@
 
 #pragma once
 
-// empty, it's added only to make UnrealBuildTool happy...
+#include "Core.h"
diff --git a/Engine/Source/Runtime/Navmesh/Private/Recast/Recast.cpp b/Engine/Source/Runtime/Navmesh/Private/Recast/Recast.cpp
--- a/Engine/Source/Runtime/Navmesh/Private/Recast/Recast.cpp
+++ b/Engine/Source/Runtime/Navmesh/Private/Recast/Recast.cpp
@@ -61,7 +61,7 @@ void rcContext::log(const rcLogCategory category, const char* format, ...)
 	char msg[MSG_SIZE];
 	va_list ap;
 	va_start(ap, format);
-	int len = vsnprintf(msg, MSG_SIZE, format, ap);
+	int len = FCStringAnsi::GetVarArgs(msg, MSG_SIZE, MSG_SIZE - 1, format, ap);
 	if (len >= MSG_SIZE)
 	{
 		len = MSG_SIZE-1;
