fixed navmesh projection on landscapes
- BVTree now includes height data from detail mesh

diff --git a/Engine/Source/Runtime/Navmesh/Private/Detour/DetourNavMeshBuilder.cpp b/Engine/Source/Runtime/Navmesh/Private/Detour/DetourNavMeshBuilder.cpp
--- a/Engine/Source/Runtime/Navmesh/Private/Detour/DetourNavMeshBuilder.cpp
+++ b/Engine/Source/Runtime/Navmesh/Private/Detour/DetourNavMeshBuilder.cpp
@@ -175,6 +175,7 @@ static void subdivide(BVItem* items, int nitems, int imin, int imax, int& curNod
 
 static int createBVTree(const unsigned short* verts, const int /*nverts*/,
 						const unsigned short* polys, const int npolys, const int nvp,
+						const dtPolyDetail* DMeshes, const float* DVerts, const unsigned char* DTris, const float* tbmin,
 						const float cs, const float ch,
 						const int /*nnodes*/, dtBVNode* nodes)
 {
@@ -190,9 +191,15 @@ static int createBVTree(const unsigned short* verts, const int /*nverts*/,
 		it.bmin[1] = it.bmax[1] = verts[p[0]*3+1];
 		it.bmin[2] = it.bmax[2] = verts[p[0]*3+2];
 		
+		int vertCount = 0;
 		for (int j = 1; j < nvp; ++j)
 		{
-			if (p[j] == MESH_NULL_IDX) break;
+			if (p[j] == MESH_NULL_IDX)
+			{
+				vertCount = j;
+				break;
+			}
+
 			unsigned short x = verts[p[j]*3+0];
 			unsigned short y = verts[p[j]*3+1];
 			unsigned short z = verts[p[j]*3+2];
@@ -205,6 +212,27 @@ static int createBVTree(const unsigned short* verts, const int /*nverts*/,
 			if (y > it.bmax[1]) it.bmax[1] = y;
 			if (z > it.bmax[2]) it.bmax[2] = z;
 		}
+
+		// include y from detail mesh
+		const dtPolyDetail* pd = &DMeshes[i];
+		for (int k = 0; k < pd->triCount; ++k)
+		{
+			const unsigned char* t = &DTris[(pd->triBase + k) * 4];
+			for (int m = 0; m < 3; ++m)
+			{
+				if (t[m] >= vertCount)
+				{
+					const float* detailCoords = &DVerts[(pd->vertBase + (t[m] - vertCount)) * 3];
+					const float qY = (detailCoords[1] - tbmin[1]) / ch;
+					const unsigned short qYmin = (unsigned short)floorf(qY);
+					const unsigned short qYmax = (unsigned short)ceilf(qY);
+
+					if (qYmin < it.bmin[1]) it.bmin[1] = qYmin;
+					if (qYmax > it.bmax[1]) it.bmax[1] = qYmax;
+				}
+			}
+		}
+
 		// Remap y
 		it.bmin[1] = (unsigned short)floorf((float)it.bmin[1]*ch/cs);
 		it.bmax[1] = (unsigned short)ceilf((float)it.bmax[1]*ch/cs);
@@ -668,11 +696,11 @@ bool dtCreateNavMeshData(dtNavMeshCreateParams* params, unsigned char** outData,
 	}
 
 	// Store and create BVtree.
-	// TODO: take detail mesh into account! use byte per bbox extent?
 	if (params->buildBvTree)
 	{
-		createBVTree(params->verts, params->vertCount, params->polys, params->polyCount,
-					 nvp, params->cs, params->ch, params->polyCount*2, navBvtree);
+		createBVTree(params->verts, params->vertCount, params->polys, params->polyCount, nvp,
+					 navDMeshes, navDVerts, navDTris, params->bmin,
+					 params->cs, params->ch, params->polyCount*2, navBvtree);
 	}
 	
 	// Store Off-Mesh connections.
