Made navmesh query filters virtual by default #UE4

- this addresses TTP#338974
- there's still a way to have non-virtual filters, configurable via ARecastNavMesh::bUseVirtualFilters
- also, can be controller per filter instance

diff --git a/Engine/Source/Runtime/Engine/Classes/AI/Navigation/RecastNavMesh.h b/Engine/Source/Runtime/Engine/Classes/AI/Navigation/RecastNavMesh.h
--- a/Engine/Source/Runtime/Engine/Classes/AI/Navigation/RecastNavMesh.h
+++ b/Engine/Source/Runtime/Engine/Classes/AI/Navigation/RecastNavMesh.h
@@ -402,6 +402,15 @@ struct FTileSetItem
 
 DECLARE_MULTICAST_DELEGATE(FOnNavMeshUpdate);
 
+namespace FNavMeshConfig
+{
+	struct FRecastNamedFiltersCreator
+	{
+		FRecastNamedFiltersCreator(bool bVirtualFilters);
+	};
+}
+
+
 UCLASS(config=Engine, defaultconfig, hidecategories=(Input,Rendering,Tags,Transform,"Utilities|Transformation",Actor,Layers,Replication))
 class ENGINE_API ARecastNavMesh : public ANavigationData
 {
@@ -518,6 +527,10 @@ class ENGINE_API ARecastNavMesh : public ANavigationData
 	UPROPERTY(config)
 	float DefaultMaxHierarchicalSearchNodes;
 
+	/** Indicates whether default navigation filters will use virtual functions. Defaults to true. */
+	UPROPERTY(config)
+	uint32 bUseVirtualFilters : 1;
+
 	/** partitioning method for creating navmesh polys */
 	UPROPERTY(EditAnywhere, Category=Generation, config, AdvancedDisplay)
 	TEnumAsByte<ERecastPartitioning::Type> RegionPartitioning;
diff --git a/Engine/Source/Runtime/Engine/Private/AI/Navigation/PImplRecastNavMesh.cpp b/Engine/Source/Runtime/Engine/Private/AI/Navigation/PImplRecastNavMesh.cpp
--- a/Engine/Source/Runtime/Engine/Private/AI/Navigation/PImplRecastNavMesh.cpp
+++ b/Engine/Source/Runtime/Engine/Private/AI/Navigation/PImplRecastNavMesh.cpp
@@ -124,7 +124,8 @@ ENavigationQueryResult::Type DTStatusToNavQueryResult(dtStatus Status)
 // FRecastQueryFilter();
 //----------------------------------------------------------------------//
 
-FRecastQueryFilter::FRecastQueryFilter()
+FRecastQueryFilter::FRecastQueryFilter(bool bIsVirtual)
+	: dtQueryFilter(bIsVirtual)
 {
 	SetExcludedArea(RECAST_NULL_AREA);
 }
@@ -134,10 +135,17 @@ INavigationQueryFilterInterface* FRecastQueryFilter::CreateCopy() const
 	return new FRecastQueryFilter(*this);
 }
 
+void FRecastQueryFilter::SetIsVirtual(bool bIsVirtual)
+{
+	dtQueryFilter* Filter = static_cast<dtQueryFilter*>(this);
+	Filter = new(Filter)dtQueryFilter(bIsVirtual);
+	SetExcludedArea(RECAST_NULL_AREA);
+}
+
 void FRecastQueryFilter::Reset()
 {
 	dtQueryFilter* Filter = static_cast<dtQueryFilter*>(this);
-	Filter = new(Filter) dtQueryFilter();
+	Filter = new(Filter) dtQueryFilter(isVirtual);
 	SetExcludedArea(RECAST_NULL_AREA);
 }
 
diff --git a/Engine/Source/Runtime/Engine/Private/AI/Navigation/PImplRecastNavMesh.h b/Engine/Source/Runtime/Engine/Private/AI/Navigation/PImplRecastNavMesh.h
--- a/Engine/Source/Runtime/Engine/Private/AI/Navigation/PImplRecastNavMesh.h
+++ b/Engine/Source/Runtime/Engine/Private/AI/Navigation/PImplRecastNavMesh.h
@@ -18,7 +18,7 @@
 class ENGINE_API FRecastQueryFilter : public INavigationQueryFilterInterface, public dtQueryFilter
 {
 public:
-	FRecastQueryFilter();
+	FRecastQueryFilter(bool bIsVirtual = true);
 	virtual ~FRecastQueryFilter(){}
 
 	virtual void Reset() override;
@@ -38,6 +38,9 @@ public:
 	virtual INavigationQueryFilterInterface* CreateCopy() const override;
 
 	const dtQueryFilter* GetAsDetourQueryFilter() const { return this; }
+
+	/** note that it results in loosing all area cost setup. Call it before setting anything else */
+	void SetIsVirtual(bool bIsVirtual);
 };
 
 struct ENGINE_API FRecastSpeciaLinkFilter : public dtQuerySpecialLinkFilter
diff --git a/Engine/Source/Runtime/Engine/Private/AI/Navigation/RecastNavMesh.cpp b/Engine/Source/Runtime/Engine/Private/AI/Navigation/RecastNavMesh.cpp
--- a/Engine/Source/Runtime/Engine/Private/AI/Navigation/RecastNavMesh.cpp
+++ b/Engine/Source/Runtime/Engine/Private/AI/Navigation/RecastNavMesh.cpp
@@ -194,17 +194,19 @@ const FRecastQueryFilter* ARecastNavMesh::NamedFilters[] = {
 	, &ERecastNamedFilter::FilterOutNavLinksAndAreasImpl
 };
 
-namespace
-{
-struct FRecastNamedFiltersCreator
+namespace FNavMeshConfig
 {
-	ARecastNavMesh::FNavPolyFlags NavLinkFlag;
+	ARecastNavMesh::FNavPolyFlags NavLinkFlag = ARecastNavMesh::FNavPolyFlags(0);
 
-	FRecastNamedFiltersCreator()
+	FRecastNamedFiltersCreator::FRecastNamedFiltersCreator(bool bVirtualFilters)
 	{
 		// setting up the last bit available in dtPoly::flags
 		NavLinkFlag = ARecastNavMesh::FNavPolyFlags(1 << (sizeof(((dtPoly*)0)->flags) * 8 - 1));
 
+		ERecastNamedFilter::FilterOutNavLinksImpl.SetIsVirtual(bVirtualFilters);
+		ERecastNamedFilter::FilterOutAreasImpl.SetIsVirtual(bVirtualFilters);
+		ERecastNamedFilter::FilterOutNavLinksAndAreasImpl.SetIsVirtual(bVirtualFilters);
+
 		ERecastNamedFilter::FilterOutNavLinksImpl.setExcludeFlags(NavLinkFlag);
 		ERecastNamedFilter::FilterOutNavLinksAndAreasImpl.setExcludeFlags(NavLinkFlag);
 
@@ -217,11 +219,9 @@ struct FRecastNamedFiltersCreator
 		ERecastNamedFilter::FilterOutAreasImpl.setAreaCost(RECAST_DEFAULT_AREA, 1.f);
 		ERecastNamedFilter::FilterOutNavLinksAndAreasImpl.setAreaCost(RECAST_DEFAULT_AREA, 1.f);
 	}
-};
 }
-static const FRecastNamedFiltersCreator RecastNamedFiltersCreator;
 
-ARecastNavMesh::FNavPolyFlags ARecastNavMesh::NavLinkFlag = RecastNamedFiltersCreator.NavLinkFlag;
+ARecastNavMesh::FNavPolyFlags ARecastNavMesh::NavLinkFlag = ARecastNavMesh::FNavPolyFlags(0);
 
 ARecastNavMesh::ARecastNavMesh(const class FPostConstructInitializeProperties& PCIP)
 	: Super(PCIP)
@@ -232,6 +232,7 @@ ARecastNavMesh::ARecastNavMesh(const class FPostConstructInitializeProperties& P
 	, MaxSimplificationError(1.3f)	// from RecastDemo
 	, DefaultMaxSearchNodes(RECAST_MAX_SEARCH_NODES)
 	, DefaultMaxHierarchicalSearchNodes(RECAST_MAX_SEARCH_NODES)
+	, bUseVirtualFilters(true)
 	, bPerformVoxelFiltering(true)	
 	, NextTimeToSortTiles(0.f)
 	, TileSetUpdateInterval(1.0f)
@@ -264,7 +265,7 @@ ARecastNavMesh::ARecastNavMesh(const class FPostConstructInitializeProperties& P
 		RaycastImplementation = NavMeshRaycast;
 
 		RecastNavMeshImpl = new FPImplRecastNavMesh(this);
-		DefaultQueryFilter->SetFilterType<FRecastQueryFilter>();
+		
 		TickHelper.Owner = this;
 
 		// add default and null areas up front
@@ -686,6 +687,9 @@ void ARecastNavMesh::PostInitProperties()
 	if (HasAnyFlags(RF_ClassDefaultObject) == true)
 	{
 		SetDrawDistance(DefaultDrawDistance);
+
+		static const FNavMeshConfig::FRecastNamedFiltersCreator RecastNamedFiltersCreator(bUseVirtualFilters);
+		NavLinkFlag = FNavMeshConfig::NavLinkFlag;
 	}
 	else if(GetWorld() != NULL)
 	{
@@ -709,10 +713,11 @@ void ARecastNavMesh::PostInitProperties()
 		NavDataReadWriteLock = &NavDataReadWriteLockDummy;		
 #endif // RECAST_ASYNC_REBUILDING
 
-		DefaultQueryFilter->SetMaxSearchNodes(DefaultMaxSearchNodes);
-
+		DefaultQueryFilter->SetFilterType<FRecastQueryFilter>();
 		FRecastQueryFilter* DetourFilter = static_cast<FRecastQueryFilter*>(DefaultQueryFilter->GetImplementation());
+		DetourFilter->SetIsVirtual(bUseVirtualFilters);
 		DetourFilter->setHeuristicScale(HeuristicScale);
+		DefaultQueryFilter->SetMaxSearchNodes(DefaultMaxSearchNodes);		
 	}
 
 	// voxel cache requires the same rasterization setup for all navmeshes, as it's stored in octree
diff --git a/Engine/Source/Runtime/Navmesh/Private/Detour/DetourNavMeshQuery.cpp b/Engine/Source/Runtime/Navmesh/Private/Detour/DetourNavMeshQuery.cpp
--- a/Engine/Source/Runtime/Navmesh/Private/Detour/DetourNavMeshQuery.cpp
+++ b/Engine/Source/Runtime/Navmesh/Private/Detour/DetourNavMeshQuery.cpp
@@ -52,7 +52,7 @@
 ///
 /// <b>Custom Implementations</b>
 /// 
-/// DT_VIRTUAL_QUERYFILTER must be defined in order to extend this class.
+/// dtQueryFilter.isVIrtual must be true in order to extend this class.
 /// 
 /// Implement a custom query filter by overriding the virtual passFilter() 
 /// and getCost() functions. If this is done, both functions should be as 
@@ -87,62 +87,6 @@ void dtQueryFilterData::copyFrom(const dtQueryFilterData* source)
 {
 	memcpy((void*)this, source, sizeof(dtQueryFilterData));
 }
-
-#ifdef DT_VIRTUAL_QUERYFILTER
-bool dtQueryFilter::passFilter(const dtPolyRef /*ref*/,
-							   const dtMeshTile* /*tile*/,
-							   const dtPoly* poly) const
-{
-	return (poly->flags & data.m_includeFlags) != 0 && (poly->flags & data.m_excludeFlags) == 0
-		&& (data.m_areaCost[poly->getArea()] < DT_UNWALKABLE_POLY_COST)
-#if WITH_FIXED_AREA_ENTERING_COST
-		&& (data.m_areaFixedCost[poly->getArea()] < DT_UNWALKABLE_POLY_COST)
-#endif // WITH_FIXED_AREA_ENTERING_COST
-		;
-}
-
-float dtQueryFilter::getCost(const float* pa, const float* pb,
-							 const dtPolyRef /*prevRef*/, const dtMeshTile* /*prevTile*/, const dtPoly* /*prevPoly*/,
-							 const dtPolyRef /*curRef*/, const dtMeshTile* /*curTile*/, const dtPoly* curPoly,
-							 const dtPolyRef /*nextRef*/, const dtMeshTile* /*nextTile*/, const dtPoly* nextPoly) const
-{
-//@UE4 BEGIN
-#if WITH_FIXED_AREA_ENTERING_COST
-	const float areaChangeCost = nextPoly != 0 && nextPoly->getArea() != curPoly->getArea() 
-		? data.m_areaFixedCost[nextPoly->getArea()] : 0.f;
-	
-	return dtVdist(pa, pb) * data.m_areaCost[curPoly->getArea()] + areaChangeCost;
-#else
-//@UE4 END
-	return dtVdist(pa, pb) * data.m_areaCost[curPoly->getArea()];
-//@UE4 BEGIN
-#endif // #if WITH_FIXED_AREA_ENTERING_COST
-//@UE4 END
-}
-#else
-//@UE4 BEGIN
-// dtQueryFilter::passFilter moved to header file due to the cross-lib inlining issues
-//@UE4 END
-
-inline float dtQueryFilter::getCost(const float* pa, const float* pb,
-									const dtPolyRef /*prevRef*/, const dtMeshTile* /*prevTile*/, const dtPoly* /*prevPoly*/,
-									const dtPolyRef /*curRef*/, const dtMeshTile* /*curTile*/, const dtPoly* curPoly,
-									const dtPolyRef /*nextRef*/, const dtMeshTile* /*nextTile*/, const dtPoly* nextPoly) const
-{
-//@UE4 BEGIN
-#if WITH_FIXED_AREA_ENTERING_COST
-	const float areaChangeCost = nextPoly != 0 && nextPoly->getArea() != curPoly->getArea() 
-		? data.m_areaFixedCost[nextPoly->getArea()] : 0.f;
-
-	return dtVdist(pa, pb) * data.m_areaCost[curPoly->getArea()] + areaChangeCost;
-#else
-//@UE4 END
-	return dtVdist(pa, pb) * data.m_areaCost[curPoly->getArea()];
-//@UE4 BEGIN
-#endif // #if WITH_FIXED_AREA_ENTERING_COST
-//@UE4 END
-}
-#endif	
 	
 //@UE4 BEGIN
 // removed following line to make H_SCALE parametrizable (via dtQueryFilter::heuristicScale)
@@ -1296,7 +1240,17 @@ dtStatus dtNavMeshQuery::findPath(dtPolyRef startRef, dtPolyRef endRef,
 			float heuristic = 0;
 			
 			// Special case for last node.
-			if (neighbourRef == endRef)
+			if (neighbourRef != endRef)
+			{
+				// Cost
+				const float curCost = filter->getCost(bestNode->pos, neiPos,
+					parentRef, parentTile, parentPoly,
+					bestRef, bestTile, bestPoly,
+					neighbourRef, neighbourTile, neighbourPoly);
+				cost = bestNode->cost + curCost;
+				heuristic = dtVdist(neiPos, endPos)*H_SCALE;
+			}
+			else
 			{
 				// Cost
 				const float curCost = filter->getCost(bestNode->pos, neiPos,
@@ -1311,16 +1265,6 @@ dtStatus dtNavMeshQuery::findPath(dtPolyRef startRef, dtPolyRef endRef,
 				cost = bestNode->cost + curCost + endCost;
 				heuristic = 0;
 			}
-			else
-			{
-				// Cost
-				const float curCost = filter->getCost(bestNode->pos, neiPos,
-													  parentRef, parentTile, parentPoly,
-													  bestRef, bestTile, bestPoly,
-													  neighbourRef, neighbourTile, neighbourPoly);
-				cost = bestNode->cost + curCost;
-				heuristic = dtVdist(neiPos, endPos)*H_SCALE;
-			}
 
 			const float total = cost + heuristic;
 			
@@ -1762,7 +1706,17 @@ dtStatus dtNavMeshQuery::updateSlicedFindPath(const int maxIter, int* doneIters)
 			float heuristic = 0;
 			
 			// Special case for last node.
-			if (neighbourRef == m_query.endRef)
+			if (neighbourRef != m_query.endRef)
+			{
+				// Cost
+				const float curCost = m_query.filter->getCost(bestNode->pos, neiPos,
+					parentRef, parentTile, parentPoly,
+					bestRef, bestTile, bestPoly,
+					neighbourRef, neighbourTile, neighbourPoly);
+				cost = bestNode->cost + curCost;
+				heuristic = dtVdist(neiPos, m_query.endPos)*H_SCALE;
+			}
+			else
 			{
 				// Cost
 				const float curCost = m_query.filter->getCost(bestNode->pos, neiPos,
@@ -1777,16 +1731,6 @@ dtStatus dtNavMeshQuery::updateSlicedFindPath(const int maxIter, int* doneIters)
 				cost = bestNode->cost + curCost + endCost;
 				heuristic = 0;
 			}
-			else
-			{
-				// Cost
-				const float curCost = m_query.filter->getCost(bestNode->pos, neiPos,
-															  parentRef, parentTile, parentPoly,
-															  bestRef, bestTile, bestPoly,
-															  neighbourRef, neighbourTile, neighbourPoly);
-				cost = bestNode->cost + curCost;
-				heuristic = dtVdist(neiPos, m_query.endPos)*H_SCALE;
-			}
 			
 			const float total = cost + heuristic;
 			
@@ -4129,13 +4073,11 @@ bool dtNavMeshQuery::isValidPolyRef(dtPolyRef ref, const dtQueryFilter* filter)
 	const dtMeshTile* tile = 0;
 	const dtPoly* poly = 0;
 	dtStatus status = m_nav->getTileAndPolyByRef(ref, &tile, &poly);
-	// If cannot get polygon, assume it does not exists and boundary is invalid.
-	if (dtStatusFailed(status))
-		return false;
-	// If cannot pass filter, assume flags has changed and boundary is invalid.
-	if (!filter->passFilter(ref, tile, poly) || !passLinkFilterByRef(tile, ref))
-		return false;
-	return true;
+
+	// should be able to get the polygon if the boundary is valid
+	return !dtStatusFailed(status)
+		// and should pass all filters
+		&& filter->passFilter(ref, tile, poly) && passLinkFilterByRef(tile, ref);
 }
 
 /// @par
diff --git a/Engine/Source/Runtime/Navmesh/Public/Detour/DetourNavMeshQuery.h b/Engine/Source/Runtime/Navmesh/Public/Detour/DetourNavMeshQuery.h
--- a/Engine/Source/Runtime/Navmesh/Public/Detour/DetourNavMeshQuery.h
+++ b/Engine/Source/Runtime/Navmesh/Public/Detour/DetourNavMeshQuery.h
@@ -24,6 +24,7 @@
 
 #include "DetourNavMesh.h"
 #include "DetourStatus.h"
+#include "DetourCommon.h"
 
 //@UE4 BEGIN
 #ifndef WITH_FIXED_AREA_ENTERING_COST
@@ -36,12 +37,11 @@
 #define DT_UNWALKABLE_POLY_COST FLT_MAX
 //@UE4 END
 
-// Define DT_VIRTUAL_QUERYFILTER if you wish to derive a custom filter from dtQueryFilter.
+// By default dtQueryFilter will use virtual calls.
 // On certain platforms indirect or virtual function call is expensive. The default
 // setting is to use non-virtual functions, the actual implementations of the functions
 // are declared as inline for maximum speed. 
-
-//#define DT_VIRTUAL_QUERYFILTER 1
+// To avoid virtual calls create dtQueryFilter with inIsVirtual = false.
 
 // Special link filter is custom filter run only for offmesh links with assigned UserId
 // Used by smart navlinks in UE4
@@ -81,24 +81,19 @@ class NAVMESH_API dtQueryFilter
 {
 protected:
 	dtQueryFilterData data;
-
+	const bool isVirtual;
+	
 public:
+	dtQueryFilter(bool inIsVirtual = true) : isVirtual(inIsVirtual)
+	{}
 	virtual ~dtQueryFilter() {}
 	
-	/// Returns true if the polygon can be visited.  (I.e. Is traversable.)
-	///  @param[in]		ref		The reference id of the polygon test.
-	///  @param[in]		tile	The tile containing the polygon.
-	///  @param[in]		poly  The polygon to test.
-#ifdef DT_VIRTUAL_QUERYFILTER
-	virtual bool passFilter(const dtPolyRef ref,
-							const dtMeshTile* tile,
-							const dtPoly* poly) const;
-#else
-//@UE4 BEGIN
-	// moved here from cpp file due to the cross-lib inlining issues
-	inline bool passFilter(const dtPolyRef /*ref*/,
-			const dtMeshTile* /*tile*/,
-			const dtPoly* poly) const
+protected:
+
+	/// inlined filter implementation. @see passFilter for parameter description
+	inline bool passInlineFilter(const dtPolyRef ref,
+		const dtMeshTile* tile,
+		const dtPoly* poly) const
 	{
 		return (poly->flags & data.m_includeFlags) != 0 && (poly->flags & data.m_excludeFlags) == 0
 			&& (data.m_areaCost[poly->getArea()] < DT_UNWALKABLE_POLY_COST)
@@ -107,9 +102,58 @@ public:
 #endif // WITH_FIXED_AREA_ENTERING_COST
 			;
 	}
-//@UE4 END
-#endif
 
+	/// virtual filter implementation (defaults to passInlineFilter). @see passFilter for parameter description
+	virtual bool passVirtualFilter(const dtPolyRef ref,
+		const dtMeshTile* tile,
+		const dtPoly* poly) const
+	{
+		return passInlineFilter(ref, tile, poly);
+	}
+
+public:
+	/// Returns true if the polygon can be visited.  (I.e. Is traversable.)
+	///  @param[in]		ref		The reference id of the polygon test.
+	///  @param[in]		tile	The tile containing the polygon.
+	///  @param[in]		poly  The polygon to test.
+	inline bool passFilter(const dtPolyRef ref,
+							const dtMeshTile* tile,
+							const dtPoly* poly) const
+	{
+		return !isVirtual ? passInlineFilter(ref, tile, poly) : passVirtualFilter(ref, tile, poly);
+	}
+
+protected:
+
+	/// inlined scoring function. @see getCost for parameter description
+	inline float getInlineCost(const float* pa, const float* pb,
+		const dtPolyRef prevRef, const dtMeshTile* prevTile, const dtPoly* prevPoly,
+		const dtPolyRef curRef, const dtMeshTile* curTile, const dtPoly* curPoly,
+		const dtPolyRef nextRef, const dtMeshTile* nextTile, const dtPoly* nextPoly) const
+	{
+#if WITH_FIXED_AREA_ENTERING_COST
+		const float areaChangeCost = nextPoly != 0 && nextPoly->getArea() != curPoly->getArea()
+			? data.m_areaFixedCost[nextPoly->getArea()] : 0.f;
+
+		return dtVdist(pa, pb) * data.m_areaCost[curPoly->getArea()] + areaChangeCost;
+#else
+		return dtVdist(pa, pb) * data.m_areaCost[curPoly->getArea()];
+#endif // #if WITH_FIXED_AREA_ENTERING_COST
+	}
+
+	/// virtual scoring function implementation (defaults to getInlineCost). @see getCost for parameter description
+	virtual float getVirtualCost(const float* pa, const float* pb,
+		const dtPolyRef prevRef, const dtMeshTile* prevTile, const dtPoly* prevPoly,
+		const dtPolyRef curRef, const dtMeshTile* curTile, const dtPoly* curPoly,
+		const dtPolyRef nextRef, const dtMeshTile* nextTile, const dtPoly* nextPoly) const
+	{
+		return getInlineCost(pa, pb,
+			prevRef, prevTile, prevPoly,
+			curRef, curTile, curPoly,
+			nextRef, nextTile, nextPoly);
+	}
+
+public:
 	/// Returns cost to move from the beginning to the end of a line segment
 	/// that is fully contained within a polygon.
 	///  @param[in]		pa			The start position on the edge of the previous and current polygon. [(x, y, z)]
@@ -120,20 +164,17 @@ public:
 	///  @param[in]		curRef		The reference id of the current polygon.
 	///  @param[in]		curTile		The tile containing the current polygon.
 	///  @param[in]		curPoly		The current polygon.
-	///  @param[in]		nextRef		The refernece id of the next polygon. [opt]
+	///  @param[in]		nextRef		The reference id of the next polygon. [opt]
 	///  @param[in]		nextTile	The tile containing the next polygon. [opt]
 	///  @param[in]		nextPoly	The next polygon. [opt]
-#ifdef DT_VIRTUAL_QUERYFILTER
-	virtual float getCost(const float* pa, const float* pb,
+	inline float getCost(const float* pa, const float* pb,
 						  const dtPolyRef prevRef, const dtMeshTile* prevTile, const dtPoly* prevPoly,
 						  const dtPolyRef curRef, const dtMeshTile* curTile, const dtPoly* curPoly,
-						  const dtPolyRef nextRef, const dtMeshTile* nextTile, const dtPoly* nextPoly) const;
-#else
-	float getCost(const float* pa, const float* pb,
-				  const dtPolyRef prevRef, const dtMeshTile* prevTile, const dtPoly* prevPoly,
-				  const dtPolyRef curRef, const dtMeshTile* curTile, const dtPoly* curPoly,
-				  const dtPolyRef nextRef, const dtMeshTile* nextTile, const dtPoly* nextPoly) const;
-#endif
+						  const dtPolyRef nextRef, const dtMeshTile* nextTile, const dtPoly* nextPoly) const
+	{
+		return !isVirtual ? getInlineCost(pa, pb, prevRef, prevTile, prevPoly, curRef, curTile, curPoly, nextRef, nextTile, nextPoly)
+			: getVirtualCost(pa, pb, prevRef, prevTile, prevPoly, curRef, curTile, curPoly, nextRef, nextTile, nextPoly);
+	}
 
 	/// @name Getters and setters for the default implementation data.
 	///@{
@@ -635,19 +676,12 @@ private:
 
 	inline bool passLinkFilter(const dtMeshTile* tile, const int polyIdx) const
 	{
-		// check if poly is offmesh link with assigned UserId
-		if (m_linkFilter && polyIdx >= tile->header->offMeshBase)
-		{
-			const int linkIdx = polyIdx - tile->header->offMeshBase;
-			if (linkIdx < tile->header->offMeshConCount &&
-				tile->offMeshCons[linkIdx].userId != 0 &&
-				m_linkFilter->isLinkAllowed(tile->offMeshCons[linkIdx].userId) == false)
-			{
-				return false;
-			}
-		}
-
-		return true;
+		const int linkIdx = polyIdx - tile->header->offMeshBase;
+
+		return !(m_linkFilter && polyIdx >= tile->header->offMeshBase
+			&& linkIdx < tile->header->offMeshConCount
+			&& tile->offMeshCons[linkIdx].userId != 0
+			&& m_linkFilter->isLinkAllowed(tile->offMeshCons[linkIdx].userId) == false);
 	}
 
 	const dtNavMesh* m_nav;							///< Pointer to navmesh data.
