Reverse order of tags in comments. Looks like valid markdown synax for UnrealDocTool; causes an error and prevents the class being documented.

diff --git a/Engine/Source/Runtime/Navmesh/Public/Detour/DetourNavMeshBuilder.h b/Engine/Source/Runtime/Navmesh/Public/Detour/DetourNavMeshBuilder.h
--- a/Engine/Source/Runtime/Navmesh/Public/Detour/DetourNavMeshBuilder.h
+++ b/Engine/Source/Runtime/Navmesh/Public/Detour/DetourNavMeshBuilder.h
@@ -120,7 +120,7 @@ struct dtNavMeshCreateParams
 	unsigned int userId;	///< The user defined id of the tile.
 	int tileX;				///< The tile's x-grid location within the multi-tile destination mesh. (Along the x-axis.)
 	int tileY;				///< The tile's y-grid location within the multi-tile desitation mesh. (Along the z-axis.)
-	int tileLayer;			///< The tile's layer within the layered destination mesh. [Limit: >= 0] (Along the y-axis.)
+	int tileLayer;			///< The tile's layer within the layered destination mesh. (Along the y-axis.) [Limit: >= 0]
 	float bmin[3];			///< The minimum bounds of the tile. [(x, y, z)] [Unit: wu]
 	float bmax[3];			///< The maximum bounds of the tile. [(x, y, z)] [Unit: wu]
 
