Fixed memory corruption in navigation mesh during world origin rebasing

diff --git a/Engine/Source/Runtime/Navmesh/Private/Detour/DetourNavMesh.cpp b/Engine/Source/Runtime/Navmesh/Private/Detour/DetourNavMesh.cpp
--- a/Engine/Source/Runtime/Navmesh/Private/Detour/DetourNavMesh.cpp
+++ b/Engine/Source/Runtime/Navmesh/Private/Detour/DetourNavMesh.cpp
@@ -2384,8 +2384,8 @@ void dtNavMesh::applyWorldOffset(const float* offset)
 			//Shift off-mesh connections
 			for (int j = 0; j < tile.header->offMeshConCount; ++j)
 			{
-				dtVadd(&(tile.offMeshCons->pos[j*6+0]), &(tile.offMeshCons->pos[j*6+0]), offset);
-				dtVadd(&(tile.offMeshCons->pos[j*6+3]), &(tile.offMeshCons->pos[j*6+3]), offset);
+				dtVadd(&(tile.offMeshCons[j].pos[0]), &(tile.offMeshCons[j].pos[0]), offset);
+				dtVadd(&(tile.offMeshCons[j].pos[3]), &(tile.offMeshCons[j].pos[3]), offset);
 			}
 			
 			// Shift clusters
