fixed crowd simulation skipping custom links, changed navlink proximity test to 2D + Z for better accuracy

diff --git a/Engine/Source/Runtime/AIModule/Private/Navigation/CrowdFollowingComponent.cpp b/Engine/Source/Runtime/AIModule/Private/Navigation/CrowdFollowingComponent.cpp
--- a/Engine/Source/Runtime/AIModule/Private/Navigation/CrowdFollowingComponent.cpp
+++ b/Engine/Source/Runtime/AIModule/Private/Navigation/CrowdFollowingComponent.cpp
@@ -457,7 +457,7 @@ int32 UCrowdFollowingComponent::DetermineStartingPathPoint(const FNavigationPath
 		// no path = called from SwitchToNextPathPart
 		if (ConsideredPath == NULL && Path.IsValid())
 		{
-			StartIdx = LastPathPolyIndex + 1;
+			StartIdx = LastPathPolyIndex;
 		}
 	}
 	else
diff --git a/Engine/Source/Runtime/Navmesh/Private/DetourCrowd/DetourCrowd.cpp b/Engine/Source/Runtime/Navmesh/Private/DetourCrowd/DetourCrowd.cpp
--- a/Engine/Source/Runtime/Navmesh/Private/DetourCrowd/DetourCrowd.cpp
+++ b/Engine/Source/Runtime/Navmesh/Private/DetourCrowd/DetourCrowd.cpp
@@ -1735,9 +1735,11 @@ void dtCrowd::updateStepOffMeshVelocity(const float dt, dtCrowdAgentDebugInfo*)
 
 		anim->t += dt;
 		
-		const float dist = dtVdistSqr(ag->npos, anim->endPos);
+		const float dist = dtVdist2DSqr(ag->npos, anim->endPos);
 		const float distThres = dtSqr(5.0f);
-		if (dist < distThres)
+		const float heightDiff = dtAbs(ag->npos[1] - anim->endPos[1]);
+		const float heightThres = ag->params.height * 0.5f;
+		if (dist < distThres && heightDiff < heightThres)
 		{
 			// Reset animation
 			anim->active = 0;
