fixed navmesh contour generation on small tiles (increased max iterations x2)

diff --git a/Engine/Source/Runtime/Navmesh/Private/DetourTileCache/DetourTileCacheBuilder.cpp b/Engine/Source/Runtime/Navmesh/Private/DetourTileCache/DetourTileCacheBuilder.cpp
--- a/Engine/Source/Runtime/Navmesh/Private/DetourTileCache/DetourTileCacheBuilder.cpp
+++ b/Engine/Source/Runtime/Navmesh/Private/DetourTileCache/DetourTileCacheBuilder.cpp
@@ -253,7 +253,7 @@ static bool walkContour(dtTileCacheLayer& layer, int x, int y, int idx, unsigned
 	unsigned short neiReg = 0xffff;
 	unsigned char neiArea = 0;
 
-	const int maxIter = w*h;
+	const int maxIter = w * h * 2;
 	int iter = 0;
 	while (iter < maxIter)
 	{
