added auto scaling of pathfinding heuristic weight to work with area costs lower than 1
UE-7936

diff --git a/Engine/Source/Runtime/Navmesh/Private/Detour/DetourNavMeshQuery.cpp b/Engine/Source/Runtime/Navmesh/Private/Detour/DetourNavMeshQuery.cpp
--- a/Engine/Source/Runtime/Navmesh/Private/Detour/DetourNavMeshQuery.cpp
+++ b/Engine/Source/Runtime/Navmesh/Private/Detour/DetourNavMeshQuery.cpp
@@ -68,7 +68,7 @@
 ///
 /// @see dtNavMeshQuery
 
-dtQueryFilterData::dtQueryFilterData() : heuristicScale(0.999f), m_includeFlags(0xffff), m_excludeFlags(0), m_isBacktracking(0)
+dtQueryFilterData::dtQueryFilterData() : heuristicScale(0.999f), lowestAreaCost(1.0f), m_includeFlags(0xffff), m_excludeFlags(0), m_isBacktracking(0)
 {
 	for (int i = 0; i < DT_MAX_AREAS; ++i)
 		m_areaCost[i] = 1.0f;
@@ -1263,7 +1263,7 @@ dtStatus dtNavMeshQuery::findPath(dtPolyRef startRef, dtPolyRef endRef,
 	}
 	
 	//@UE4 BEGIN
-	const float H_SCALE = filter->getHeuristicScale();
+	const float H_SCALE = filter->getModifiedHeuristicScale();
 	//@UE4 END
 
 #if TRACK_PATHFINDING_PERF
@@ -1673,7 +1673,7 @@ dtStatus dtNavMeshQuery::initSlicedFindPath(dtPolyRef startRef, dtPolyRef endRef
 	}
 	
 	//@UE4 BEGIN
-	const float H_SCALE = filter->getHeuristicScale();
+	const float H_SCALE = filter->getModifiedHeuristicScale();
 	//@UE4 END
 
 	m_nodePool->clear();
@@ -1708,7 +1708,7 @@ dtStatus dtNavMeshQuery::updateSlicedFindPath(const int maxIter, int* doneIters)
 	}
 
 	//@UE4 BEGIN
-	const float H_SCALE = m_query.filter->getHeuristicScale();
+	const float H_SCALE = m_query.filter->getModifiedHeuristicScale();
 	//@UE4 END
 			
 	int iter = 0;
diff --git a/Engine/Source/Runtime/Navmesh/Public/Detour/DetourNavMeshQuery.h b/Engine/Source/Runtime/Navmesh/Public/Detour/DetourNavMeshQuery.h
--- a/Engine/Source/Runtime/Navmesh/Public/Detour/DetourNavMeshQuery.h
+++ b/Engine/Source/Runtime/Navmesh/Public/Detour/DetourNavMeshQuery.h
@@ -63,6 +63,7 @@ struct dtQueryFilterData
 	float m_areaFixedCost[DT_MAX_AREAS];///< Fixed cost for entering an area type (Used by default implementation.)
 #endif // WITH_FIXED_AREA_ENTERING_COST
 	float heuristicScale;				///< Search heuristic scale.
+	float lowestAreaCost;
 
 	unsigned short m_includeFlags;		///< Flags for polygons that can be visited. (Used by default implementation.)
 	unsigned short m_excludeFlags;		///< Flags for polygons that should not be visted. (Used by default implementation.)
@@ -187,7 +188,7 @@ public:
 	/// Sets the traversal cost of the area.
 	///  @param[in]		i		The id of the area.
 	///  @param[in]		cost	The new cost of traversing the area.
-	inline void setAreaCost(const int i, const float cost) { data.m_areaCost[i] = cost; }
+	inline void setAreaCost(const int i, const float cost) { data.m_areaCost[i] = cost; data.lowestAreaCost = dtMin(data.lowestAreaCost, cost); }
 
 //@UE4 BEGIN
 	inline const float* getAllAreaCosts() const { return data.m_areaCost; }
@@ -206,6 +207,8 @@ public:
 	inline const float* getAllFixedAreaCosts() const { return data.m_areaFixedCost; }
 #endif // WITH_FIXED_AREA_ENTERING_COST
 
+	inline float getModifiedHeuristicScale() const { return data.heuristicScale * ((data.lowestAreaCost > 0) ? data.lowestAreaCost : 1.0f); }
+
 	/// Retrieves euclidean distance heuristic scale
 	///  @returns heuristic scale
 	inline float getHeuristicScale() const { return data.heuristicScale; }
