fix for cycles in pathfinding when heuristic weight is overestimating distance to goal
UE-8037

diff --git a/Engine/Source/Runtime/Navmesh/Private/Detour/DetourNavMeshQuery.cpp b/Engine/Source/Runtime/Navmesh/Private/Detour/DetourNavMeshQuery.cpp
--- a/Engine/Source/Runtime/Navmesh/Private/Detour/DetourNavMeshQuery.cpp
+++ b/Engine/Source/Runtime/Navmesh/Private/Detour/DetourNavMeshQuery.cpp
@@ -1348,14 +1348,21 @@ dtStatus dtNavMeshQuery::findPath(dtPolyRef startRef, dtPolyRef endRef,
 				status |= DT_OUT_OF_NODES;
 				continue;
 			}
-			
-			// Always calculate correct position on neighbor's edge,
-			// skipping to wrong edge may greatly change path cost
-			// (if area cost differences are more than 5x default)
+
+			// Try to update node position for current edge to make paths more precise
+			// Unless heuristic is not admissible (overestimates), in which case
+			// we have to use constant values for given node to avoid creating cycles
 			float neiPos[3] = { 0.0f, 0.0f, 0.0f };
-			getEdgeMidPoint(bestRef, bestPoly, bestTile,
-				neighbourRef, neighbourPoly, neighbourTile,
-				neiPos);
+			if (H_SCALE <= 1.0f || neighbourNode->flags == 0)
+			{
+				getEdgeMidPoint(bestRef, bestPoly, bestTile,
+					neighbourRef, neighbourPoly, neighbourTile,
+					neiPos);
+			}
+			else
+			{
+				dtVcopy(neiPos, neighbourNode->pos);
+			}
 
 			// Calculate cost and heuristic.
 			float cost = 0;
