added fixed multiplier for crowd's avoidance query range (early detection)
#ue4

diff --git a/Engine/Source/Runtime/AIModule/Classes/Navigation/CrowdFollowingComponent.h b/Engine/Source/Runtime/AIModule/Classes/Navigation/CrowdFollowingComponent.h
--- a/Engine/Source/Runtime/AIModule/Classes/Navigation/CrowdFollowingComponent.h
+++ b/Engine/Source/Runtime/AIModule/Classes/Navigation/CrowdFollowingComponent.h
@@ -76,6 +76,7 @@ class AIMODULE_API UCrowdFollowingComponent : public UPathFollowingComponent, pu
 	void SetCrowdCollisionQueryRange(float Range, bool bUpdateAgent = true);
 	void SetCrowdPathOptimizationRange(float Range, bool bUpdateAgent = true);
 	void SetCrowdAvoidanceQuality(ECrowdAvoidanceQuality::Type Quality, bool bUpdateAgent = true);
+	void SetCrowdAvoidanceRangeMultiplier(float Multipler, bool bUpdateAgent = true);
 
 	FORCEINLINE bool IsCrowdSimulationEnabled() const { return bEnableCrowdSimulation; }
 	FORCEINLINE bool IsCrowdSimulatioSuspended() const { return bSuspendCrowdSimulation; }
@@ -101,6 +102,7 @@ class AIMODULE_API UCrowdFollowingComponent : public UPathFollowingComponent, pu
 	FORCEINLINE float GetCrowdCollisionQueryRange() const { return CollisionQueryRange; }
 	FORCEINLINE float GetCrowdPathOptimizationRange() const { return PathOptimizationRange; }
 	FORCEINLINE ECrowdAvoidanceQuality::Type GetCrowdAvoidanceQuality() const { return AvoidanceQuality; }
+	FORCEINLINE float GetCrowdAvoidanceRangeMultiplier() const { return AvoidanceRangeMultiplier; }
 	FORCEINLINE int32 GetAvoidanceGroup() const { return AvoidanceGroup.Packed; }
 	FORCEINLINE int32 GetGroupsToAvoid() const { return GroupsToAvoid.Packed; }
 	FORCEINLINE int32 GetGroupsToIgnore() const { return GroupsToIgnore.Packed; }
@@ -160,6 +162,9 @@ protected:
 	float CollisionQueryRange;
 	float PathOptimizationRange;
 
+	/** multiplier for avoidance samples during detection, doesn't affect actual velocity */
+	float AvoidanceRangeMultiplier;
+
 	/** start index of current path part */
 	int32 PathStartIndex;
 
diff --git a/Engine/Source/Runtime/AIModule/Private/Navigation/CrowdFollowingComponent.cpp b/Engine/Source/Runtime/AIModule/Private/Navigation/CrowdFollowingComponent.cpp
--- a/Engine/Source/Runtime/AIModule/Private/Navigation/CrowdFollowingComponent.cpp
+++ b/Engine/Source/Runtime/AIModule/Private/Navigation/CrowdFollowingComponent.cpp
@@ -29,6 +29,7 @@ UCrowdFollowingComponent::UCrowdFollowingComponent(const FObjectInitializer& Obj
 	CollisionQueryRange = 400.0f;		// approx: radius * 12.0f
 	PathOptimizationRange = 1000.0f;	// approx: radius * 30.0f
 	AvoidanceQuality = ECrowdAvoidanceQuality::Low;
+	AvoidanceRangeMultiplier = 1.0f;
 
 	AvoidanceGroup.SetFlagsDirectly(1);
 	GroupsToAvoid.SetFlagsDirectly(MAX_uint32);
diff --git a/Engine/Source/Runtime/AIModule/Private/Navigation/CrowdManager.cpp b/Engine/Source/Runtime/AIModule/Private/Navigation/CrowdManager.cpp
--- a/Engine/Source/Runtime/AIModule/Private/Navigation/CrowdManager.cpp
+++ b/Engine/Source/Runtime/AIModule/Private/Navigation/CrowdManager.cpp
@@ -692,6 +692,7 @@ void UCrowdManager::GetAgentParams(const ICrowdAgentInterface* Agent, dtCrowdAge
 		AgentParams.pathOptimizationRange = CrowdComponent->GetCrowdPathOptimizationRange();
 		AgentParams.separationWeight = CrowdComponent->GetCrowdSeparationWeight();
 		AgentParams.obstacleAvoidanceType = CrowdComponent->GetCrowdAvoidanceQuality();
+		AgentParams.avoidanceQueryMultiplier = CrowdComponent->GetCrowdAvoidanceRangeMultiplier();
 	
 		if (CrowdComponent->IsCrowdSimulationEnabled())
 		{
@@ -711,6 +712,7 @@ void UCrowdManager::GetAgentParams(const ICrowdAgentInterface* Agent, dtCrowdAge
 	}
 	else
 	{
+		AgentParams.avoidanceQueryMultiplier = 1.0f;
 		AgentParams.avoidanceGroup = 1;
 		AgentParams.groupsToAvoid = MAX_uint32;
 	}
@@ -957,13 +959,14 @@ void UCrowdManager::DrawDebugVelocityObstacles(const dtCrowdAgent* CrowdAgent) c
 	FVector Center = Recast2UnrealPoint(CrowdAgent->npos) + CrowdDebugDrawing::Offset;
 	DrawDebugCylinder(GetWorld(), Center - CrowdDebugDrawing::Offset, Center, CrowdAgent->params.maxSpeed, 32, CrowdDebugDrawing::AvoidanceRange);
 
+	const float InvQueryMultiplier = 1.0f / CrowdAgent->params.avoidanceQueryMultiplier;
 	float BestSampleScore = -1.0f;
 	FVector BestSampleLocation = FVector::ZeroVector;
 
 	for (int32 Idx = 0; Idx < DetourAvoidanceDebug->getSampleCount(); Idx++)
 	{
 		const float* p = DetourAvoidanceDebug->getSampleVelocity(Idx);
-		const float sr = DetourAvoidanceDebug->getSampleSize(Idx);
+		const float sr = DetourAvoidanceDebug->getSampleSize(Idx) * InvQueryMultiplier;
 		const float pen = DetourAvoidanceDebug->getSamplePenalty(Idx);
 		const float pen2 = DetourAvoidanceDebug->getSamplePreferredSidePenalty(Idx);
 
diff --git a/Engine/Source/Runtime/Navmesh/Private/DetourCrowd/DetourCrowd.cpp b/Engine/Source/Runtime/Navmesh/Private/DetourCrowd/DetourCrowd.cpp
--- a/Engine/Source/Runtime/Navmesh/Private/DetourCrowd/DetourCrowd.cpp
+++ b/Engine/Source/Runtime/Navmesh/Private/DetourCrowd/DetourCrowd.cpp
@@ -1559,8 +1559,8 @@ void dtCrowd::updateStepAvoidance(const float dt, dtCrowdAgentDebugInfo* debug)
 
 			// Sample new safe velocity.
 			const dtObstacleAvoidanceParams* params = &m_obstacleQueryParams[ag->params.obstacleAvoidanceType];
-			const int ns = m_obstacleQuery->sampleVelocity(
-					ag->npos, ag->params.radius, ag->desiredSpeed,
+			const int ns = m_obstacleQuery->sampleVelocity(ag->npos, ag->params.radius,
+					ag->desiredSpeed, ag->params.avoidanceQueryMultiplier,
 					ag->vel, ag->dvel, ag->nvel, params, vod);
 
 			m_velocitySampleCount += ns;
diff --git a/Engine/Source/Runtime/Navmesh/Private/DetourCrowd/DetourObstacleAvoidance.cpp b/Engine/Source/Runtime/Navmesh/Private/DetourCrowd/DetourObstacleAvoidance.cpp
--- a/Engine/Source/Runtime/Navmesh/Private/DetourCrowd/DetourObstacleAvoidance.cpp
+++ b/Engine/Source/Runtime/Navmesh/Private/DetourCrowd/DetourObstacleAvoidance.cpp
@@ -444,7 +444,8 @@ bool dtObstacleAvoidanceQuery::getCustomSamplingPattern(int idx, float* angles,
 	return true;
 }
 
-int dtObstacleAvoidanceQuery::sampleVelocityCustom(const float* pos, const float rad, const float vmax,
+int dtObstacleAvoidanceQuery::sampleVelocityCustom(const float* pos, const float rad,
+												   const float vmax, const float vmult,
 												   const float* vel, const float* dvel, float* nvel,
 												   const dtObstacleAvoidanceParams* params,
 											  	   dtObstacleAvoidanceDebugData* debug)
@@ -473,7 +474,7 @@ int dtObstacleAvoidanceQuery::sampleVelocityCustom(const float* pos, const float
 	}
 
 	float minPenalty = FLT_MAX;
-	float cr = vmax * (1.0f - m_params.velBias);
+	float cr = vmax * vmult * (1.0f - m_params.velBias);
 	float res[3];
 	bool bFoundSample = false;
 	dtVset(res, dvel[0] * m_params.velBias, 0, dvel[2] * m_params.velBias);
@@ -485,7 +486,7 @@ int dtObstacleAvoidanceQuery::sampleVelocityCustom(const float* pos, const float
 		vcand[1] = 0;
 		vcand[2] = res[2] + pat[i * 2 + 1] * cr;
 
-		if (dtSqr(vcand[0]) + dtSqr(vcand[2]) > dtSqr(vmax + 0.001f)) continue;
+		if (dtSqr(vcand[0]) + dtSqr(vcand[2]) > dtSqr((vmax * vmult) + 0.001f)) continue;
 
 		const float penalty = processSample(vcand, 20.0f, pos, rad, vel, dvel, debug);
 		if (penalty < minPenalty)
@@ -500,11 +501,16 @@ int dtObstacleAvoidanceQuery::sampleVelocityCustom(const float* pos, const float
 	{
 		dtVcopy(nvel, dvel);
 	}
+	else
+	{
+		dtVscale(nvel, nvel, 1.0f / vmult);
+	}
 
 	return pattern.nsamples;
 }
 
-int dtObstacleAvoidanceQuery::sampleVelocityAdaptive(const float* pos, const float rad, const float vmax,
+int dtObstacleAvoidanceQuery::sampleVelocityAdaptive(const float* pos, const float rad,
+													 const float vmax, const float vmult,
 													 const float* vel, const float* dvel, float* nvel,
 													 const dtObstacleAvoidanceParams* params,
 													 dtObstacleAvoidanceDebugData* debug)
@@ -513,8 +519,8 @@ int dtObstacleAvoidanceQuery::sampleVelocityAdaptive(const float* pos, const flo
 	
 	memcpy(&m_params, params, sizeof(dtObstacleAvoidanceParams));
 	m_invHorizTime = 1.0f / m_params.horizTime;
-	m_vmax = vmax;
-	m_invVmax = 1.0f / vmax;
+	m_vmax = vmax * vmult;
+	m_invVmax = 1.0f / (vmax * vmult);
 	
 	dtVset(nvel, 0,0,0);
 	
@@ -553,11 +559,12 @@ int dtObstacleAvoidanceQuery::sampleVelocityAdaptive(const float* pos, const flo
 	}
 
 	// Start sampling.
-	float cr = vmax * (1.0f - m_params.velBias);
+	float cr = vmax * vmult * (1.0f - m_params.velBias);
 	float res[3];
 	dtVset(res, dvel[0] * m_params.velBias, 0, dvel[2] * m_params.velBias);
 	int ns = 0;
 
+	const float invVmult = 1.0f / vmult;
 	for (int k = 0; k < depth; ++k)
 	{
 		float minPenalty = FLT_MAX;
@@ -571,7 +578,7 @@ int dtObstacleAvoidanceQuery::sampleVelocityAdaptive(const float* pos, const flo
 			vcand[1] = 0;
 			vcand[2] = res[2] + pat[i*2+1]*cr;
 			
-			if (dtSqr(vcand[0])+dtSqr(vcand[2]) > dtSqr(vmax+0.001f)) continue;
+			if (dtSqr(vcand[0])+dtSqr(vcand[2]) > dtSqr((vmax * vmult)+0.001f)) continue;
 			
 			const float penalty = processSample(vcand,cr/10, pos,rad,vel,dvel, debug);
 			ns++;
@@ -582,7 +589,7 @@ int dtObstacleAvoidanceQuery::sampleVelocityAdaptive(const float* pos, const flo
 			}
 		}
 
-		dtVcopy(res, bvel);
+		dtVscale(res, bvel, invVmult);
 
 		cr *= 0.5f;
 	}	
diff --git a/Engine/Source/Runtime/Navmesh/Public/DetourCrowd/DetourCrowd.h b/Engine/Source/Runtime/Navmesh/Public/DetourCrowd/DetourCrowd.h
--- a/Engine/Source/Runtime/Navmesh/Public/DetourCrowd/DetourCrowd.h
+++ b/Engine/Source/Runtime/Navmesh/Public/DetourCrowd/DetourCrowd.h
@@ -94,6 +94,9 @@ struct dtCrowdAgentParams
 	/// How aggresive the agent manager should be at avoiding collisions with this agent. [Limit: >= 0]
 	float separationWeight;
 
+	/// [UE4] Mutliplier for avoidance velocities
+	float avoidanceQueryMultiplier;
+
 	/// [UE4] Groups flags attached to the agent
 	unsigned int avoidanceGroup;
 	/// [UE4] Avoid agents when they group is matching mask
diff --git a/Engine/Source/Runtime/Navmesh/Public/DetourCrowd/DetourObstacleAvoidance.h b/Engine/Source/Runtime/Navmesh/Public/DetourCrowd/DetourObstacleAvoidance.h
--- a/Engine/Source/Runtime/Navmesh/Public/DetourCrowd/DetourObstacleAvoidance.h
+++ b/Engine/Source/Runtime/Navmesh/Public/DetourCrowd/DetourObstacleAvoidance.h
@@ -123,25 +123,28 @@ public:
 	bool getCustomSamplingPattern(int idx, float* angles, float* radii, int* nsamples);
 
 	// [UE4] sample velocity using custom patterns
-	int sampleVelocityCustom(const float* pos, const float rad, const float vmax,
+	int sampleVelocityCustom(const float* pos, const float rad,
+					 		 const float vmax, const float vmult,
 							 const float* vel, const float* dvel, float* nvel,
 							 const dtObstacleAvoidanceParams* params,
 							 dtObstacleAvoidanceDebugData* debug = 0);
 
-	int sampleVelocityAdaptive(const float* pos, const float rad, const float vmax,
+	int sampleVelocityAdaptive(const float* pos, const float rad,
+							   const float vmax, const float vmult,
 							   const float* vel, const float* dvel, float* nvel,
 							   const dtObstacleAvoidanceParams* params, 
 							   dtObstacleAvoidanceDebugData* debug = 0);
 	
 	// [UE4] main sampling function
-	inline int sampleVelocity(const float* pos, const float rad, const float vmax,
+	inline int sampleVelocity(const float* pos, const float rad,
+		const float vmax, const float vmult,
 		const float* vel, const float* dvel, float* nvel,
 		const dtObstacleAvoidanceParams* params,
 		dtObstacleAvoidanceDebugData* debug = 0)
 	{
 		return (params->patternIdx == 0xff) ?
-			sampleVelocityAdaptive(pos, rad, vmax, vel, dvel, nvel, params, debug) :
-			sampleVelocityCustom(pos, rad, vmax, vel, dvel, nvel, params, debug);
+			sampleVelocityAdaptive(pos, rad, vmax, vmult, vel, dvel, nvel, params, debug) :
+			sampleVelocityCustom(pos, rad, vmax, vmult, vel, dvel, nvel, params, debug);
 	}
 
 	inline int getObstacleCircleCount() const { return m_ncircles; }
