added support for MinRegionArea param in navmesh generation
UE-8375 #ue4

diff --git a/Engine/Source/Runtime/Engine/Classes/AI/Navigation/RecastNavMesh.h b/Engine/Source/Runtime/Engine/Classes/AI/Navigation/RecastNavMesh.h
--- a/Engine/Source/Runtime/Engine/Classes/AI/Navigation/RecastNavMesh.h
+++ b/Engine/Source/Runtime/Engine/Classes/AI/Navigation/RecastNavMesh.h
@@ -502,7 +502,7 @@ class ENGINE_API ARecastNavMesh : public ANavigationData
 	UPROPERTY(EditAnywhere, Category=Generation, config, meta=(ClampMin = "0.0"))
 	float MinRegionArea;
 
-	/* The size limit of regions to be merged with bigger regions */
+	/* The size limit of regions to be merged with bigger regions (watershed partitioning only) */
 	UPROPERTY(EditAnywhere, Category=Generation, config, meta=(ClampMin = "0.0"))
 	float MergeRegionSize;
 
diff --git a/Engine/Source/Runtime/Engine/Private/AI/Navigation/RecastNavMeshGenerator.cpp b/Engine/Source/Runtime/Engine/Private/AI/Navigation/RecastNavMeshGenerator.cpp
--- a/Engine/Source/Runtime/Engine/Private/AI/Navigation/RecastNavMeshGenerator.cpp
+++ b/Engine/Source/Runtime/Engine/Private/AI/Navigation/RecastNavMeshGenerator.cpp
@@ -2259,7 +2259,7 @@ bool FRecastTileGenerator::GenerateNavigationData(FNavMeshBuildContext& BuildCon
 			// Build regions
 			if (TileConfig.TileCachePartitionType == RC_REGION_MONOTONE)
 			{
-				status = dtBuildTileCacheRegionsMonotone(&MyAllocator, *GenerationContext.Layer);
+				status = dtBuildTileCacheRegionsMonotone(&MyAllocator, TileConfig.minRegionArea, TileConfig.mergeRegionArea, *GenerationContext.Layer);
 			}
 			else if (TileConfig.TileCachePartitionType == RC_REGION_WATERSHED)
 			{
@@ -2277,12 +2277,11 @@ bool FRecastTileGenerator::GenerateNavigationData(FNavMeshBuildContext& BuildCon
 					return false;
 				}
 
-				const int TileBoderSize = 0;
-				status = dtBuildTileCacheRegions(&MyAllocator, TileBoderSize, TileConfig.minRegionArea, TileConfig.mergeRegionArea, *GenerationContext.Layer, *GenerationContext.DistanceField);
+				status = dtBuildTileCacheRegions(&MyAllocator, TileConfig.minRegionArea, TileConfig.mergeRegionArea, *GenerationContext.Layer, *GenerationContext.DistanceField);
 			}
 			else
 			{
-				status = dtBuildTileCacheRegionsChunky(&MyAllocator, *GenerationContext.Layer, TileConfig.TileCacheChunkSize);
+				status = dtBuildTileCacheRegionsChunky(&MyAllocator, TileConfig.minRegionArea, TileConfig.mergeRegionArea, *GenerationContext.Layer, TileConfig.TileCacheChunkSize);
 			}
 
 			if (dtStatusFailed(status))
diff --git a/Engine/Source/Runtime/Navmesh/Private/DetourTileCache/DetourTileCache.cpp b/Engine/Source/Runtime/Navmesh/Private/DetourTileCache/DetourTileCache.cpp
--- a/Engine/Source/Runtime/Navmesh/Private/DetourTileCache/DetourTileCache.cpp
+++ b/Engine/Source/Runtime/Navmesh/Private/DetourTileCache/DetourTileCache.cpp
@@ -661,7 +661,7 @@ dtStatus dtTileCache::buildNavMeshTile(const dtCompressedTileRef ref, dtNavMesh*
 	// Build navmesh
 	if (m_params.regionPartitioning == DT_REGION_MONOTONE)
 	{
-		status = dtBuildTileCacheRegionsMonotone(m_talloc, *bc.layer);
+		status = dtBuildTileCacheRegionsMonotone(m_talloc, m_params.minRegionArea, m_params.mergeRegionArea, *bc.layer);
 	}
 	else if (m_params.regionPartitioning == DT_REGION_WATERSHED)
 	{
@@ -673,11 +673,11 @@ dtStatus dtTileCache::buildNavMeshTile(const dtCompressedTileRef ref, dtNavMesh*
 		if (dtStatusFailed(status))
 			return status;
 
-		status = dtBuildTileCacheRegions(m_talloc, 0, m_params.minRegionArea, m_params.mergeRegionArea, *bc.layer, *bc.dfield);
+		status = dtBuildTileCacheRegions(m_talloc, m_params.minRegionArea, m_params.mergeRegionArea, *bc.layer, *bc.dfield);
 	}
 	else
 	{
-		status = dtBuildTileCacheRegionsChunky(m_talloc, *bc.layer, m_params.regionChunkSize);
+		status = dtBuildTileCacheRegionsChunky(m_talloc, m_params.minRegionArea, m_params.mergeRegionArea, *bc.layer, m_params.regionChunkSize);
 	}
 	if (dtStatusFailed(status))
 		return status;
diff --git a/Engine/Source/Runtime/Navmesh/Private/DetourTileCache/DetourTileCacheRegion.cpp b/Engine/Source/Runtime/Navmesh/Private/DetourTileCache/DetourTileCacheRegion.cpp
--- a/Engine/Source/Runtime/Navmesh/Private/DetourTileCache/DetourTileCacheRegion.cpp
+++ b/Engine/Source/Runtime/Navmesh/Private/DetourTileCache/DetourTileCacheRegion.cpp
@@ -29,8 +29,6 @@
 #define _USE_MATH_DEFINES
 #include <math.h>
 
-static const unsigned short DT_BORDER_REG = 0x8000;
-
 inline bool isConnected(const dtTileCacheLayer& layer, const int idx, const int dir)
 {
 	return (layer.cons[idx] & (1 << dir)) != 0;
@@ -276,23 +274,6 @@ dtStatus dtBuildTileCacheDistanceField(dtTileCacheAlloc* alloc, dtTileCacheLayer
 	return DT_SUCCESS;
 }
 
-static void paintRectRegion(int minx, int maxx, int miny, int maxy, unsigned short regId,
-	dtTileCacheLayer& layer, unsigned short* srcReg)
-{
-	const int w = (int)layer.header->width;	
-	for (int y = miny; y < maxy; ++y)
-	{
-		for (int x = minx; x < maxx; ++x)
-		{
-			const int idx = x+y*w;
-			if (layer.areas[idx] != DT_TILECACHE_NULL_AREA)
-			{
-				srcReg[idx] = regId;
-			}
-		}
-	}
-}
-
 static unsigned short* expandRegions(int maxIter, unsigned short level,
 	dtTileCacheLayer& layer, dtTileCacheDistanceField& dfield,
 	unsigned short* srcReg, unsigned short* srcDist,
@@ -348,7 +329,7 @@ static unsigned short* expandRegions(int maxIter, unsigned short level,
 				if (ax >= 0 && ax < w && ay >= 0 && ay < h && isConnected(layer, i, dir))
 				{
 					if (layer.areas[ai] != area) continue;
-					if (srcReg[ai] > 0 && (srcReg[ai] & DT_BORDER_REG) == 0)
+					if (srcReg[ai] > 0)
 					{
 						if ((int)srcDist[ai]+2 < (int)d2)
 						{
@@ -427,8 +408,6 @@ static bool floodRegion(int x, int y, int i, unsigned short level, unsigned shor
 				if (layer.areas[ai] != area)
 					continue;
 				unsigned short nr = srcReg[ai];
-				if (nr & DT_BORDER_REG) // Do not take borders into account.
-					continue;
 				if (nr != 0 && nr != r)
 					ar = nr;
 
@@ -485,7 +464,8 @@ struct dtLayerRegion
 		id(i),
 		areaType(0),
 		remap(false),
-		visited(false)
+		visited(false),
+		border(false)
 	{}
 
 	dtIntArray connections;
@@ -494,6 +474,7 @@ struct dtLayerRegion
 	unsigned char areaType;			// Are type.
 	unsigned char remap : 1;
 	unsigned char visited : 1;
+	unsigned char border : 1;
 
 	dtLayerRegion& operator=(const dtLayerRegion& src)
 	{
@@ -502,6 +483,7 @@ struct dtLayerRegion
 		this->areaType = src.areaType;
 		this->remap = src.remap;
 		this->visited = src.visited;
+		this->border = src.border;
 		this->connections.copy(src.connections);
 		return *this;
 	}
@@ -605,23 +587,15 @@ static bool mergeRegions(dtLayerRegion& rega, dtLayerRegion& regb)
 
 	rega.cellCount += regb.cellCount;
 	regb.cellCount = 0;
+
+	rega.border |= regb.border;
+	regb.border = 0;
+
 	regb.connections.resize(0);
 
 	return true;
 }
 
-static bool isRegionConnectedToBorder(const dtLayerRegion& reg)
-{
-	// Region is connected to border if
-	// one of the neighbours is null id.
-	for (int i = 0; i < reg.connections.size(); ++i)
-	{
-		if (reg.connections[i] == 0)
-			return true;
-	}
-	return false;
-}
-
 static bool isSolidEdge(dtTileCacheLayer& layer, unsigned short* srcReg,
 	int x, int y, int i, int dir)
 {
@@ -747,6 +721,7 @@ static dtStatus filterSmallRegions(dtTileCacheAlloc* alloc, dtTileCacheLayer& la
 	// Find edge of a region and find connections around the contour.
 	for (int y = 0; y < h; ++y)
 	{
+		const bool borderY = (y == 0) || (y == (h - 1));
 		for (int x = 0; x < w; ++x)
 		{
 			const int i = x+y*w;
@@ -757,6 +732,7 @@ static dtStatus filterSmallRegions(dtTileCacheAlloc* alloc, dtTileCacheLayer& la
 
 			dtLayerRegion& reg = regions[r];
 			reg.cellCount++;
+			reg.border |= borderY || (x == 0) || (x == (w - 1));
 
 			// Have found contour
 			if (reg.connections.size() > 0)
@@ -790,7 +766,7 @@ static dtStatus filterSmallRegions(dtTileCacheAlloc* alloc, dtTileCacheLayer& la
 	for (int i = 0; i < nreg; ++i)
 	{
 		dtLayerRegion& reg = regions[i];
-		if (reg.id == 0 || (reg.id & DT_BORDER_REG))
+		if (reg.id == 0)
 			continue;                       
 		if (reg.cellCount == 0)
 			continue;
@@ -814,20 +790,16 @@ static dtStatus filterSmallRegions(dtTileCacheAlloc* alloc, dtTileCacheLayer& la
 
 			dtLayerRegion& creg = regions[ri];
 
+			connectsToBorder |= creg.border;
 			cellCount += creg.cellCount;
 			trace.push(ri);
 
 			for (int j = 0; j < creg.connections.size(); ++j)
 			{
-				if (creg.connections[j] & DT_BORDER_REG)
-				{
-					connectsToBorder = true;
-					continue;
-				}
 				dtLayerRegion& neireg = regions[creg.connections[j]];
 				if (neireg.visited)
 					continue;
-				if (neireg.id == 0 || (neireg.id & DT_BORDER_REG))
+				if (neireg.id == 0)
 					continue;
 				// Visit
 				stack.push(neireg.id);
@@ -858,13 +830,13 @@ static dtStatus filterSmallRegions(dtTileCacheAlloc* alloc, dtTileCacheLayer& la
 		for (int i = 0; i < nreg; ++i)
 		{
 			dtLayerRegion& reg = regions[i];
-			if (reg.id == 0 || (reg.id & DT_BORDER_REG))
+			if (reg.id == 0)
 				continue;                       
 			if (reg.cellCount == 0)
 				continue;
 
 			// Check to see if the region should be merged.
-			if (reg.cellCount > mergeRegionSize && isRegionConnectedToBorder(reg))
+			if (reg.cellCount > mergeRegionSize && reg.border)
 				continue;
 
 			// Small region with more than 1 connection.
@@ -874,9 +846,8 @@ static dtStatus filterSmallRegions(dtTileCacheAlloc* alloc, dtTileCacheLayer& la
 			unsigned short mergeId = reg.id;
 			for (int j = 0; j < reg.connections.size(); ++j)
 			{
-				if (reg.connections[j] & DT_BORDER_REG) continue;
 				dtLayerRegion& mreg = regions[reg.connections[j]];
-				if (mreg.id == 0 || (mreg.id & DT_BORDER_REG)) continue;
+				if (mreg.id == 0) continue;
 				if (mreg.cellCount < smallest &&
 					canMergeWithRegion(reg, mreg) &&
 					canMergeWithRegion(mreg, reg))
@@ -897,7 +868,7 @@ static dtStatus filterSmallRegions(dtTileCacheAlloc* alloc, dtTileCacheLayer& la
 					// Fixup regions pointing to current region.
 					for (int j = 0; j < nreg; ++j)
 					{
-						if (regions[j].id == 0 || (regions[j].id & DT_BORDER_REG)) continue;
+						if (regions[j].id == 0) continue;
 						// If another region was already merged into current region
 						// change the nid of the previous region too.
 						if (regions[j].id == oldId)
@@ -918,7 +889,6 @@ static dtStatus filterSmallRegions(dtTileCacheAlloc* alloc, dtTileCacheLayer& la
 	{
 		regions[i].remap = false;
 		if (regions[i].id == DT_TILECACHE_NULL_AREA) continue;       // Skip nil regions.
-		if (regions[i].id & DT_BORDER_REG) continue;    // Skip external regions.
 		regions[i].remap = true;
 	}
 
@@ -943,8 +913,7 @@ static dtStatus filterSmallRegions(dtTileCacheAlloc* alloc, dtTileCacheLayer& la
 	// Remap regions.
 	for (int i = w*h-1; i >= 0; i--)
 	{
-		if ((srcReg[i] & DT_BORDER_REG) == 0)
-			srcReg[i] = regions[srcReg[i]].id;
+		srcReg[i] = regions[srcReg[i]].id;
 	}
 
 	for (int i = 0; i < nreg; ++i)
@@ -954,7 +923,7 @@ static dtStatus filterSmallRegions(dtTileCacheAlloc* alloc, dtTileCacheLayer& la
 }
 
 dtStatus dtBuildTileCacheRegions(dtTileCacheAlloc* alloc,
-	const int borderSize, const int minRegionArea, const int mergeRegionArea,
+	const int minRegionArea, const int mergeRegionArea,
 	dtTileCacheLayer& layer, dtTileCacheDistanceField dfield)
 {
 	dtAssert(alloc);
@@ -989,18 +958,6 @@ dtStatus dtBuildTileCacheRegions(dtTileCacheAlloc* alloc,
 	//	const int expandIters = 4 + walkableRadius * 2;
 	const int expandIters = 8;
 
-	if (borderSize > 0)
-	{
-		// Make sure border will not overflow.
-		const int bw = dtMin(w, borderSize);
-		const int bh = dtMin(h, borderSize);
-		// Paint regions
-		paintRectRegion(0, bw, 0, h, regionId|DT_BORDER_REG, layer, srcReg); regionId++;
-		paintRectRegion(w-bw, w, 0, h, regionId|DT_BORDER_REG, layer, srcReg); regionId++;
-		paintRectRegion(0, w, 0, bh, regionId|DT_BORDER_REG, layer, srcReg); regionId++;
-		paintRectRegion(0, w, h-bh, h, regionId|DT_BORDER_REG, layer, srcReg); regionId++;
-	}
-
 	while (level > 0)
 	{
 		level = level >= 2 ? level-2 : 0;
@@ -1061,6 +1018,8 @@ struct dtLayerMonotoneRegion
 	unsigned short regId;
 	unsigned char areaId;
 	unsigned char remap : 1;
+	unsigned char border : 1;
+	unsigned char visited : 1;
 };
 
 static void addUniqueLast(dtIntArray& a, unsigned short v)
@@ -1207,6 +1166,7 @@ static dtStatus CollectRegionsMonotone(dtTileCacheAlloc* alloc, dtTileCacheLayer
 	// Find region neighbours.
 	for (int y = 0; y < h; ++y)
 	{
+		const bool borderY = (y == 0) || (y == (h - 1));
 		for (int x = 0; x < w; ++x)
 		{
 			const int idx = x+y*w;
@@ -1217,6 +1177,7 @@ static dtStatus CollectRegionsMonotone(dtTileCacheAlloc* alloc, dtTileCacheLayer
 			// Update area.
 			regs[ri].area++;
 			regs[ri].areaId = layer.areas[idx];
+			regs[ri].border |= borderY || (x == 0) || (x == (w - 1));
 
 			// Update neighbours
 			if (y > 0 && isConnected(layer, idx, 3))
@@ -1366,6 +1327,7 @@ static dtStatus CollectRegionsChunky(dtTileCacheAlloc* alloc, dtTileCacheLayer&
 	for (int y = 0; y < h; ++y)
 	{
 		const int chunkYOffset = (y / chunkSize) * chunkSize;
+		const bool borderY = (y == 0) || (y == (h - 1));
 		for (int x = 0; x < w; ++x)
 		{
 			const int idx = x+y*w;
@@ -1377,6 +1339,7 @@ static dtStatus CollectRegionsChunky(dtTileCacheAlloc* alloc, dtTileCacheLayer&
 			regs[ri].area++;
 			regs[ri].areaId = layer.areas[idx];
 			regs[ri].chunkId = (x / chunkSize) + chunkYOffset;
+			regs[ri].border |= borderY || (x == 0) || (x == (w - 1));
 
 			// Update neighbours
 			if (y > 0 && isConnected(layer, idx, 3))
@@ -1395,14 +1358,79 @@ static dtStatus CollectRegionsChunky(dtTileCacheAlloc* alloc, dtTileCacheLayer&
 	return DT_SUCCESS;
 }
 
-static void MergeAndCompressRegions(dtTileCacheAlloc* alloc, dtTileCacheLayer& layer, dtLayerMonotoneRegion* regs, int nregs)
+static void MergeAndCompressRegions(dtTileCacheAlloc* alloc, dtTileCacheLayer& layer, dtLayerMonotoneRegion* regs, int nregs, const int minRegionArea, const int mergeRegionArea)
 {
 	for (int i = 0; i < nregs; ++i)
 		regs[i].regId = (unsigned short)(i + 1);
 
+	// Remove too small regions.
+	if (minRegionArea > 0)
+	{
+		dtIntArray stack(32);
+		dtIntArray trace(32);
+		for (int i = 0; i < nregs; ++i)
+		{
+			dtLayerMonotoneRegion& reg = regs[i];
+			if (reg.visited || reg.area == 0)
+				continue;
+
+			// Count the total size of all the connected regions.
+			// Also keep track of the regions connects to a tile border.
+			bool connectsToBorder = false;
+			int cellCount = 0;
+			stack.resize(0);
+			trace.resize(0);
+
+			reg.visited = true;
+			stack.push(i);
+
+			while (stack.size())
+			{
+				// Pop
+				int ri = stack.pop();
+
+				dtLayerMonotoneRegion& creg = regs[ri];
+
+				connectsToBorder |= creg.border;
+				cellCount += creg.area;
+				trace.push(ri);
+
+				for (int j = 0; j < creg.neis.size(); ++j)
+				{
+					dtLayerMonotoneRegion& neireg = regs[creg.neis[j]];
+					if (neireg.visited)
+						continue;
+					if (neireg.regId == 0)
+						continue;
+					// Visit
+					stack.push(neireg.regId - 1);
+					neireg.visited = true;
+				}
+			}
+
+			// If the accumulated regions size is too small, remove it.
+			// Do not remove areas which connect to tile borders
+			// as their size cannot be estimated correctly and removing them
+			// can potentially remove necessary areas.
+			if (cellCount < minRegionArea && !connectsToBorder)
+			{
+				// Kill all visited regions.
+				for (int j = 0; j < trace.size(); ++j)
+				{
+					regs[trace[j]].area = 0;
+					regs[trace[j]].regId = 0;
+				}
+			}
+		}
+	}
+
 	for (int i = 0; i < nregs; ++i)
 	{
 		dtLayerMonotoneRegion& reg = regs[i];
+		if (reg.regId == 0)
+			continue;
+		// don't use mergeRegionArea, it doesn't work well with monotone partitioning
+		// (results in even more long thin polys)
 
 		int merge = -1;
 		int mergea = 0;
@@ -1443,6 +1471,7 @@ static void MergeAndCompressRegions(dtTileCacheAlloc* alloc, dtTileCacheLayer& l
 		for (int i = 0; i < nregs; ++i)
 			remap[regs[i].regId] = 1;
 		// skip region id 0, it's used for skipping minRegionArea
+		remap[0] = 0;
 		for (int i = 1; i < 256; ++i)
 			if (remap[i])
 				remap[i] = ++regId;
@@ -1485,7 +1514,7 @@ static void MergeAndCompressRegions(dtTileCacheAlloc* alloc, dtTileCacheLayer& l
 	alloc->free(regs);
 }
 
-dtStatus dtBuildTileCacheRegionsMonotone(dtTileCacheAlloc* alloc, dtTileCacheLayer& layer)
+dtStatus dtBuildTileCacheRegionsMonotone(dtTileCacheAlloc* alloc, const int minRegionArea, const int mergeRegionArea, dtTileCacheLayer& layer)
 {
 	dtLayerMonotoneRegion* regs = NULL;
 	int nregs = 0;
@@ -1493,13 +1522,13 @@ dtStatus dtBuildTileCacheRegionsMonotone(dtTileCacheAlloc* alloc, dtTileCacheLay
 	dtStatus status = CollectRegionsMonotone(alloc, layer, regs, nregs);
 	if (dtStatusSucceed(status))
 	{
-		MergeAndCompressRegions(alloc, layer, regs, nregs);
+		MergeAndCompressRegions(alloc, layer, regs, nregs, minRegionArea, mergeRegionArea);
 	}
 
 	return status;
 }
 
-dtStatus dtBuildTileCacheRegionsChunky(dtTileCacheAlloc* alloc, dtTileCacheLayer& layer, int regionChunkSize)
+dtStatus dtBuildTileCacheRegionsChunky(dtTileCacheAlloc* alloc, const int minRegionArea, const int mergeRegionArea, dtTileCacheLayer& layer, int regionChunkSize)
 {
 	dtLayerMonotoneRegion* regs = NULL;
 	int nregs = 0;
@@ -1507,7 +1536,7 @@ dtStatus dtBuildTileCacheRegionsChunky(dtTileCacheAlloc* alloc, dtTileCacheLayer
 	dtStatus status = CollectRegionsChunky(alloc, layer, regionChunkSize, regs, nregs);
 	if (dtStatusSucceed(status))
 	{
-		MergeAndCompressRegions(alloc, layer, regs, nregs);
+		MergeAndCompressRegions(alloc, layer, regs, nregs, minRegionArea, mergeRegionArea);
 	}
 
 	return status;
diff --git a/Engine/Source/Runtime/Navmesh/Public/DetourTileCache/DetourTileCacheBuilder.h b/Engine/Source/Runtime/Navmesh/Public/DetourTileCache/DetourTileCacheBuilder.h
--- a/Engine/Source/Runtime/Navmesh/Public/DetourTileCache/DetourTileCacheBuilder.h
+++ b/Engine/Source/Runtime/Navmesh/Public/DetourTileCache/DetourTileCacheBuilder.h
@@ -237,13 +237,15 @@ NAVMESH_API dtStatus dtReplaceConvexArea(dtTileCacheLayer& layer, const float* o
 NAVMESH_API dtStatus dtBuildTileCacheDistanceField(dtTileCacheAlloc* alloc, dtTileCacheLayer& layer, dtTileCacheDistanceField& dfield);
 
 NAVMESH_API dtStatus dtBuildTileCacheRegions(dtTileCacheAlloc* alloc,
-								const int borderSize, const int minRegionArea, const int mergeRegionArea,
+								const int minRegionArea, const int mergeRegionArea,
 								dtTileCacheLayer& layer, dtTileCacheDistanceField dfield);
 
 NAVMESH_API dtStatus dtBuildTileCacheRegionsMonotone(dtTileCacheAlloc* alloc,
+								const int minRegionArea, const int mergeRegionArea,
 								dtTileCacheLayer& layer);
 
 NAVMESH_API dtStatus dtBuildTileCacheRegionsChunky(dtTileCacheAlloc* alloc,
+								const int minRegionArea, const int mergeRegionArea,
 								dtTileCacheLayer& layer, int regionChunkSize);
 //@UE4 END
 
