merged crowd's shared boundary optimizations from dev branch
- #2418572: improving crowd avoidance near navmesh wall segments
- #2418666: fixed crowd's shared boundary being picked even tough it doesn't contains path's polys (walking though wall bug)
- #2420168: fixed crowd agents having boundary segments between polys that belong to path (area change)
- #2435203: optimized crowd's shared boundary usage: fixed creating multiple samples for the same locations, changed containter type for faster cleanup

#ue4

diff --git a/Engine/Source/Runtime/AIModule/Private/Navigation/CrowdManager.cpp b/Engine/Source/Runtime/AIModule/Private/Navigation/CrowdManager.cpp
--- a/Engine/Source/Runtime/AIModule/Private/Navigation/CrowdManager.cpp
+++ b/Engine/Source/Runtime/AIModule/Private/Navigation/CrowdManager.cpp
@@ -864,7 +864,6 @@ void UCrowdManager::CreateCrowdManager()
 	{
 		DetourCrowd->init(MaxAgents, MaxAgentRadius, NavMeshPtr);
 		DetourCrowd->setAgentCheckInterval(NavmeshCheckInterval);
-		DetourCrowd->setOffmeshLinkSegmentRemovalRadius(MaxAgentRadius);
 		DetourCrowd->setSingleAreaVisibilityOptimization(bSingleAreaVisibilityOptimization);
 		DetourCrowd->setPruneStartedOffmeshConnections(bPruneStartedOffmeshConnections);
 
diff --git a/Engine/Source/Runtime/Navmesh/Private/Detour/DetourNavMeshQuery.cpp b/Engine/Source/Runtime/Navmesh/Private/Detour/DetourNavMeshQuery.cpp
--- a/Engine/Source/Runtime/Navmesh/Private/Detour/DetourNavMeshQuery.cpp
+++ b/Engine/Source/Runtime/Navmesh/Private/Detour/DetourNavMeshQuery.cpp
@@ -4089,7 +4089,29 @@ dtStatus dtNavMeshQuery::findWallsInNeighbourhood(dtPolyRef startRef, const floa
 		// add hard edges of poly
 		for (int i = 0; i < curPoly->vertCount; i++)
 		{
-			if (curPoly->neis[i] == 0)
+			bool bStoreEdge = (curPoly->neis[i] == 0);
+			if (curPoly->neis[i] & DT_EXT_LINK)
+			{
+				// check if external edge has valid link
+				bool bConnected = false;
+
+				unsigned int linkIdx = curPoly->firstLink;
+				while (linkIdx != DT_NULL_LINK)
+				{
+					const dtLink& link = m_nav->getLink(curTile, linkIdx);
+					linkIdx = link.next;
+
+					if (link.edge == i)
+					{
+						bConnected = true;
+						break;
+					}
+				}
+
+				bStoreEdge = !bConnected;
+			}
+
+			if (bStoreEdge)
 			{
 				storeWallSegment(curTile, curPoly, i, 
 					curRef, 0, centerPos, radiusSqr, 
diff --git a/Engine/Source/Runtime/Navmesh/Private/DetourCrowd/DetourCrowd.cpp b/Engine/Source/Runtime/Navmesh/Private/DetourCrowd/DetourCrowd.cpp
--- a/Engine/Source/Runtime/Navmesh/Private/DetourCrowd/DetourCrowd.cpp
+++ b/Engine/Source/Runtime/Navmesh/Private/DetourCrowd/DetourCrowd.cpp
@@ -350,7 +350,6 @@ dtCrowd::dtCrowd() :
 	m_maxPathResult(0),
 	m_maxAgentRadius(0),
 	m_agentStateCheckInterval(1.0f),
-	m_linkRemovalRadius(0),
 	m_velocitySampleCount(0),
 	m_navquery(0),
 	m_raycastSingleArea(0),
@@ -1246,7 +1245,7 @@ void dtCrowd::updateStepPaths(const float dt, dtCrowdAgentDebugInfo*)
 	updateTopologyOptimization(m_activeAgents, m_numActiveAgents, dt);
 }
 
-void dtCrowd::updateStepProximityData(const float dt, dtCrowdAgentDebugInfo*)
+void dtCrowd::updateStepProximityData(const float dt, dtCrowdAgentDebugInfo* debug)
 {
 	// Register agents to proximity grid.
 	m_grid->clear();
@@ -1288,11 +1287,15 @@ void dtCrowd::updateStepProximityData(const float dt, dtCrowdAgentDebugInfo*)
 			!ag->boundary.isValid(m_navquery, &m_filters[ag->params.filter]))
 		{
 			// UE4: force removing segments too close to offmesh links
-			const bool useForcedRemove = m_linkRemovalRadius > 0.0f && ag->ncorners &&
-				(ag->cornerFlags[ag->ncorners - 1] & DT_STRAIGHTPATH_OFFMESH_CONNECTION);
-
-			const float* removePos = useForcedRemove ? &ag->cornerVerts[(ag->ncorners - 1) * 3] : 0;
-			const float removeRadius = m_linkRemovalRadius;
+			float linkV0[3] = { 0.0f };
+			float linkV1[3] = { 0.0f };
+			
+			bool bHasOffmeshLink = ag->ncorners && (ag->cornerFlags[ag->ncorners - 1] & DT_STRAIGHTPATH_OFFMESH_CONNECTION);
+			if (bHasOffmeshLink)
+			{
+				const dtStatus linkStatus = m_navquery->getAttachedNavMesh()->getOffMeshConnectionPolyEndPoints(0, ag->cornerPolys[ag->ncorners - 1], 0, linkV0, linkV1);
+				bHasOffmeshLink = dtStatusSucceed(linkStatus);
+			}
 
 			// UE4: move dir for segment scoring
 			float moveDir[3] = { 0.0f };
@@ -1307,8 +1310,8 @@ void dtCrowd::updateStepProximityData(const float dt, dtCrowdAgentDebugInfo*)
 			dtVnormalize(moveDir);
 
 			ag->boundary.update(&m_sharedBoundary, sharedDataIdx, ag->npos, ag->params.collisionQueryRange,
-				removePos, removeRadius, useForcedRemove,
-				ag->corridor.getPath(), ag->corridor.getPathCount(),
+				bHasOffmeshLink, linkV0, linkV1,
+				ag->corridor.getPath(), m_raycastSingleArea ? ag->corridor.getPathCount() : 0,
 				moveDir, m_navquery, &m_filters[ag->params.filter]);
 		}
 		// Query neighbour agents
@@ -1776,11 +1779,6 @@ void dtCrowd::setAgentCheckInterval(const float t)
 	m_agentStateCheckInterval = t;
 }
 
-void dtCrowd::setOffmeshLinkSegmentRemovalRadius(const float radius)
-{
-	m_linkRemovalRadius = radius;
-}
-
 void dtCrowd::setSingleAreaVisibilityOptimization(bool bEnable)
 {
 	m_raycastSingleArea = bEnable;
diff --git a/Engine/Source/Runtime/Navmesh/Private/DetourCrowd/DetourLocalBoundary.cpp b/Engine/Source/Runtime/Navmesh/Private/DetourCrowd/DetourLocalBoundary.cpp
--- a/Engine/Source/Runtime/Navmesh/Private/DetourCrowd/DetourLocalBoundary.cpp
+++ b/Engine/Source/Runtime/Navmesh/Private/DetourCrowd/DetourLocalBoundary.cpp
@@ -88,7 +88,7 @@ void dtLocalBoundary::addSegment(const float dist, const float* s)
 }
 
 void dtLocalBoundary::update(dtPolyRef ref, const float* pos, const float collisionQueryRange,
-	const float* forcedRemovePos, float forcedRemoveRadius, bool useForcedRemove,
+	const bool bRemoveNearLink, const float* linkV0, const float* linkV1,
 	const dtPolyRef* path, const int npath,
 	const float* moveDir,
 	dtNavMeshQuery* navquery, const dtQueryFilter* filter)
@@ -130,12 +130,15 @@ void dtLocalBoundary::update(dtPolyRef ref, const float* pos, const float collis
 				continue;
 
 			// [UE4] remove segments too close to requested position
-			if (useForcedRemove)
+			if (bRemoveNearLink)
 			{
-				float tseg2;
-				const float dist2 = dtDistancePtSegSqr2D(forcedRemovePos, s, s + 3, tseg2);
-				if (dist2 < dtSqr(forcedRemoveRadius))
+				float dummyS = 0.0f;
+				float dummyT = 0.0f;
+				const bool bIntersect = dtIntersectSegSeg2D(s, s+3, linkV0, linkV1, dummyS, dummyT);
+				if (bIntersect)
+				{
 					continue;
+				}
 			}
 
 			// [UE4] include direction to segment in score
@@ -152,26 +155,39 @@ void dtLocalBoundary::update(dtPolyRef ref, const float* pos, const float collis
 
 void dtLocalBoundary::update(const dtSharedBoundary* sharedData, const int sharedIdx,
 	const float* pos, const float collisionQueryRange,
-	const float* forcedRemovePos, float forcedRemoveRadius, bool useForcedRemove,
+	const bool bRemoveNearLink, const float* linkV0, const float* linkV1,
 	const dtPolyRef* path, const int npath, const float* moveDir,
 	dtNavMeshQuery* navquery, const dtQueryFilter* filter)
 {
-	if (!sharedData || !sharedData->Data.IsValidIndex(sharedIdx))
+	if (!sharedData || !sharedData->HasSample(sharedIdx))
 	{
 		return;
 	}
 
 	const dtSharedBoundaryData& Data = sharedData->Data[sharedIdx];
 	m_npolys = FMath::Min(Data.Polys.Num(), MAX_LOCAL_POLYS);
-	for (int32 Idx = 0; Idx < m_npolys; Idx++)
+	int32 PolyIdx = 0;
+	for (auto It = Data.Polys.CreateConstIterator(); It; ++It)
 	{
-		m_polys[Idx] = Data.Polys[Idx];
+		m_polys[PolyIdx] = *It;
+		
+		PolyIdx++;
+		if (PolyIdx >= m_npolys)
+		{
+			break;
+		}
 	}
 
 	float closestPt[3] = { 0.0f };
 	float dirToSeg[3] = { 0.0f };
 	float s[6];
 
+	TSet<dtPolyRef> PathLookup;
+	for (int32 Idx = 0; Idx < npath; Idx++)
+	{
+		PathLookup.Add(path[Idx]);
+	}
+
 	m_nsegs = 0;
 	for (int32 Idx = 0; Idx < Data.Edges.Num(); Idx++)
 	{
@@ -181,12 +197,21 @@ void dtLocalBoundary::update(const dtSharedBoundary* sharedData, const int share
 			continue;
 
 		// remove segments too close to requested position
-		if (useForcedRemove)
+		if (bRemoveNearLink)
 		{
-			float tseg2;
-			const float dist2 = dtDistancePtSegSqr2D(forcedRemovePos, Data.Edges[Idx].v0, Data.Edges[Idx].v1, tseg2);
-			if (dist2 < dtSqr(forcedRemoveRadius))
+			float dummyS = 0.0f;
+			float dummyT = 0.0f;
+			const bool bIntersect = dtIntersectSegSeg2D(Data.Edges[Idx].v0, Data.Edges[Idx].v1, linkV0, linkV1, dummyS, dummyT);
+			if (bIntersect)
+			{
 				continue;
+			}
+		}
+
+		// remove segments when both sides are on path (single area trace)
+		if (PathLookup.Contains(Data.Edges[Idx].p0) && PathLookup.Contains(Data.Edges[Idx].p1))
+		{
+			continue;
 		}
 
 		// include direction to segment in score
@@ -194,7 +219,7 @@ void dtLocalBoundary::update(const dtSharedBoundary* sharedData, const int share
 		dtVsub(dirToSeg, closestPt, pos);
 		dtVnormalize(dirToSeg);
 		const float dseg = dtVdot2D(dirToSeg, moveDir);
-		const float score = distSqr * ((1.0f - dseg) * 0.5f);
+		const float score = distSqr *((1.0f - dseg) * 0.5f);
 
 		dtVcopy(s, Data.Edges[Idx].v0);
 		dtVcopy(s + 3, Data.Edges[Idx].v1);
@@ -213,7 +238,7 @@ bool dtLocalBoundary::isValid(dtNavMeshQuery* navquery, const dtQueryFilter* fil
 		if (!navquery->isValidPolyRef(m_polys[i], filter))
 			return false;
 	}
-	
+
 	return true;
 }
 
diff --git a/Engine/Source/Runtime/Navmesh/Private/DetourCrowd/DetourObstacleAvoidance.cpp b/Engine/Source/Runtime/Navmesh/Private/DetourCrowd/DetourObstacleAvoidance.cpp
--- a/Engine/Source/Runtime/Navmesh/Private/DetourCrowd/DetourObstacleAvoidance.cpp
+++ b/Engine/Source/Runtime/Navmesh/Private/DetourCrowd/DetourObstacleAvoidance.cpp
@@ -391,6 +391,12 @@ float dtObstacleAvoidanceQuery::processSample(const float* vcand, const float cs
 				continue;
 		}
 		
+		// UE4: when sample is too close to segment (navmesh wall) - disable it completely
+		if (htmin < 0.1f)
+		{
+			return -1.0f;
+		}
+
 		// Avoid less when facing walls.
 		htmin *= 2.0f;
 		
@@ -489,7 +495,7 @@ int dtObstacleAvoidanceQuery::sampleVelocityCustom(const float* pos, const float
 		if (dtSqr(vcand[0]) + dtSqr(vcand[2]) > dtSqr((vmax * vmult) + 0.001f)) continue;
 
 		const float penalty = processSample(vcand, 20.0f, pos, rad, vel, dvel, debug);
-		if (penalty < minPenalty)
+		if (penalty < minPenalty && penalty >= 0.0f)
 		{
 			bFoundSample = true;
 			minPenalty = penalty;
@@ -582,7 +588,7 @@ int dtObstacleAvoidanceQuery::sampleVelocityAdaptive(const float* pos, const flo
 			
 			const float penalty = processSample(vcand,cr/10, pos,rad,vel,dvel, debug);
 			ns++;
-			if (penalty < minPenalty)
+			if (penalty < minPenalty && penalty >= 0.0f)
 			{
 				minPenalty = penalty;
 				dtVcopy(bvel, vcand);
diff --git a/Engine/Source/Runtime/Navmesh/Private/DetourCrowd/DetourSharedBoundary.cpp b/Engine/Source/Runtime/Navmesh/Private/DetourCrowd/DetourSharedBoundary.cpp
--- a/Engine/Source/Runtime/Navmesh/Private/DetourCrowd/DetourSharedBoundary.cpp
+++ b/Engine/Source/Runtime/Navmesh/Private/DetourCrowd/DetourSharedBoundary.cpp
@@ -21,15 +21,15 @@ void dtSharedBoundary::Tick(float DeltaTime)
 	// clear unused entries
 	if (CurrentTime > NextClearTime)
 	{
-		const float MaxLifeTime = 3.0f;
+		const float MaxLifeTime = 2.0f;
 		NextClearTime = CurrentTime + MaxLifeTime;
 
-		for (int32 Idx = Data.Num() - 1; Idx >= 0; Idx--)
+		for (TSparseArray<dtSharedBoundaryData>::TIterator It(Data); It; ++It)
 		{
-			const float LastAccess = CurrentTime - Data[Idx].AccessTime;
+			const float LastAccess = CurrentTime - It->AccessTime;
 			if (LastAccess >= MaxLifeTime)
 			{
-				Data.RemoveAt(Idx);
+				It.RemoveCurrent();
 			}
 		}
 	}
@@ -37,51 +37,72 @@ void dtSharedBoundary::Tick(float DeltaTime)
 
 int32 dtSharedBoundary::CacheData(float* Center, float Radius, dtPolyRef CenterPoly, dtNavMeshQuery* NavQuery, dtQueryFilter* NavFilter)
 {
+	// bail if requested poly is not valid (e.g. rebuild in progress)
+	if (NavQuery && !NavQuery->isValidPolyRef(CenterPoly, NavFilter))
+	{
+		return -1;
+	}
+
 	Radius *= 1.5f;
 
-	const int32 ExistingIdx = FindData(Center, Radius, NavFilter);
-	if (ExistingIdx >= 0)
+	int32 DataIdx = FindData(Center, Radius, CenterPoly, NavFilter);
+	const bool bHasValidData = IsValid(DataIdx, NavQuery, NavFilter);
+	if (!bHasValidData)
 	{
-		Data[ExistingIdx].AccessTime = CurrentTime;
-		return ExistingIdx;
-	}
+		if (DataIdx >= 0)
+		{
+			// remove in next cleanup
+			Data[DataIdx].AccessTime = 0.0f;
+		}
 
-	dtSharedBoundaryData NewData;
-	dtVcopy(NewData.Center, Center);
-	NewData.Radius = Radius;
-	NewData.Filter = NavFilter;
-	NewData.SingleAreaId = 0;
-	NewData.AccessTime = CurrentTime;
+		dtSharedBoundaryData NewData;
+		dtVcopy(NewData.Center, Center);
+		NewData.Radius = Radius;
+		NewData.Filter = NavFilter;
+		NewData.SingleAreaId = 0;
+
+		FindEdges(NewData, CenterPoly, NavQuery, NavFilter);
+		DataIdx = Data.Add(NewData);
+	}
 
-	FindEdges(NewData, CenterPoly, NavQuery, NavFilter);
-	const int32 NewIdx = Data.Add(NewData);
-	return NewIdx;
+	Data[DataIdx].AccessTime = CurrentTime;
+	return DataIdx;
 }
 
 int32 dtSharedBoundary::CacheData(float* Center, float Radius, dtPolyRef CenterPoly, dtNavMeshQuery* NavQuery, uint8 SingleAreaId)
 {
-	Radius *= 1.5f;
+	SingleAreaFilter.setAreaCost(SingleAreaId, 1.0f);
 
-	const int32 ExistingIdx = FindData(Center, Radius, SingleAreaId);
-	if (ExistingIdx >= 0)
+	// bail if requested poly is not valid (e.g. rebuild in progress)
+	if (NavQuery && !NavQuery->isValidPolyRef(CenterPoly, &SingleAreaFilter))
 	{
-		Data[ExistingIdx].AccessTime = CurrentTime;
-		return ExistingIdx;
+		return -1;
 	}
 
-	dtSharedBoundaryData NewData;
-	dtVcopy(NewData.Center, Center);
-	NewData.Radius = Radius;
-	NewData.SingleAreaId = SingleAreaId;
-	NewData.AccessTime = CurrentTime;
+	Radius *= 1.5f;
 
-	SingleAreaFilter.setAreaCost(SingleAreaId, 1.0f);
+	int32 DataIdx = FindData(Center, Radius, CenterPoly, SingleAreaId);
+	const bool bHasValidData = IsValid(DataIdx, NavQuery, &SingleAreaFilter);
+	if (!bHasValidData)
+	{
+		if (DataIdx >= 0)
+		{
+			// remove in next cleanup
+			Data[DataIdx].AccessTime = 0.0f;
+		}
 
-	FindEdges(NewData, CenterPoly, NavQuery, &SingleAreaFilter);
-	const int32 NewIdx = Data.Add(NewData);
+		dtSharedBoundaryData NewData;
+		dtVcopy(NewData.Center, Center);
+		NewData.Radius = Radius;
+		NewData.SingleAreaId = SingleAreaId;
+
+		FindEdges(NewData, CenterPoly, NavQuery, &SingleAreaFilter);
+		DataIdx = Data.Add(NewData);
+	}
 
 	SingleAreaFilter.setAreaCost(SingleAreaId, DT_UNWALKABLE_POLY_COST);
-	return NewIdx;
+	Data[DataIdx].AccessTime = CurrentTime;
+	return DataIdx;
 }
 
 void dtSharedBoundary::FindEdges(dtSharedBoundaryData& Data, dtPolyRef CenterPoly, dtNavMeshQuery* NavQuery, dtQueryFilter* NavFilter)
@@ -116,10 +137,10 @@ void dtSharedBoundary::FindEdges(dtSharedBoundaryData& Data, dtPolyRef CenterPol
 	}
 }
 
-int32 dtSharedBoundary::FindData(float* Center, float Radius, dtQueryFilter* NavFilter) const
+int32 dtSharedBoundary::FindData(float* Center, float Radius, dtPolyRef ReqPoly, dtQueryFilter* NavFilter) const
 {
 	const float RadiusThr = 50.0f;
-	const float DistThrSq = FMath::Square(Radius * 0.4f);
+	const float DistThrSq = FMath::Square(Radius * 0.5f);
 
 	for (int32 Idx = 0; Idx < Data.Num(); Idx++)
 	{
@@ -128,7 +149,10 @@ int32 dtSharedBoundary::FindData(float* Center, float Radius, dtQueryFilter* Nav
 			const float DistSq = dtVdistSqr(Center, Data[Idx].Center);
 			if (DistSq <= DistThrSq && dtAbs(Data[Idx].Radius - Radius) < RadiusThr)
 			{
-				return Idx;
+				if (Data[Idx].Polys.Contains(ReqPoly))
+				{
+					return Idx;
+				}
 			}
 		}
 	}
@@ -136,9 +160,9 @@ int32 dtSharedBoundary::FindData(float* Center, float Radius, dtQueryFilter* Nav
 	return -1;
 }
 
-int32 dtSharedBoundary::FindData(float* Center, float Radius, uint8 SingleAreaId) const
+int32 dtSharedBoundary::FindData(float* Center, float Radius, dtPolyRef ReqPoly, uint8 SingleAreaId) const
 {
-	const float DistThrSq = FMath::Square(Radius * 0.4f);
+	const float DistThrSq = FMath::Square(Radius * 0.5f);
 	const float RadiusThr = 50.0f;
 
 	for (int32 Idx = 0; Idx < Data.Num(); Idx++)
@@ -148,10 +172,38 @@ int32 dtSharedBoundary::FindData(float* Center, float Radius, uint8 SingleAreaId
 			const float DistSq = dtVdistSqr(Center, Data[Idx].Center);
 			if (DistSq <= DistThrSq && dtAbs(Data[Idx].Radius - Radius) < RadiusThr)
 			{
-				return Idx;
+				if (Data[Idx].Polys.Contains(ReqPoly))
+				{
+					return Idx;
+				}
 			}
 		}
 	}
 
 	return -1;
 }
+
+bool dtSharedBoundary::HasSample(int32 Idx) const
+{
+	return (Idx >= 0) && (Idx < Data.GetMaxIndex()) && Data.IsAllocated(Idx);
+}
+
+bool dtSharedBoundary::IsValid(int32 Idx, dtNavMeshQuery* NavQuery, dtQueryFilter* NavFilter) const
+{
+	bool bValid = HasSample(Idx);
+	if (bValid)
+	{
+		for (auto It = Data[Idx].Polys.CreateConstIterator(); It; ++It)
+		{
+			const dtPolyRef TestRef = *It;
+			const bool bValidRef = NavQuery->isValidPolyRef(TestRef, NavFilter);
+			if (!bValidRef)
+			{
+				bValid = false;
+				break;
+			}
+		}
+	}
+
+	return bValid;
+}
diff --git a/Engine/Source/Runtime/Navmesh/Public/DetourCrowd/DetourCrowd.h b/Engine/Source/Runtime/Navmesh/Public/DetourCrowd/DetourCrowd.h
--- a/Engine/Source/Runtime/Navmesh/Public/DetourCrowd/DetourCrowd.h
+++ b/Engine/Source/Runtime/Navmesh/Public/DetourCrowd/DetourCrowd.h
@@ -243,8 +243,6 @@ class NAVMESH_API dtCrowd
 
 	// [UE4] time between attempts to restore agents state
 	float m_agentStateCheckInterval;
-	// [UE4] remove segments from local boundary that are too close to offmesh links (disrupts VO)
-	float m_linkRemovalRadius;
 
 	int m_velocitySampleCount;
 
@@ -445,10 +443,6 @@ public:
 	/// @param[in]		npath	The number of polygons in the path.
 	bool setAgentCorridor(const int idx, const dtPolyRef* path, const int npath);
 
-	/// [UE4] Set radius around offmesh link start, that will remove wall segments from boundary
-	/// @param[in]		radius	Radius threshold, remove segments that are closer 
-	void setOffmeshLinkSegmentRemovalRadius(const float radius);
-
 	/// [UE4] Set visibility optimization to use single area raycasts
 	/// This will prevent from cutting through polys marked as different area
 	/// which could have been avoided in corridor's path
diff --git a/Engine/Source/Runtime/Navmesh/Public/DetourCrowd/DetourLocalBoundary.h b/Engine/Source/Runtime/Navmesh/Public/DetourCrowd/DetourLocalBoundary.h
--- a/Engine/Source/Runtime/Navmesh/Public/DetourCrowd/DetourLocalBoundary.h
+++ b/Engine/Source/Runtime/Navmesh/Public/DetourCrowd/DetourLocalBoundary.h
@@ -53,14 +53,14 @@ public:
 
 	// [UE4: new sections: link removal, path corridor, direction]
 	void update(dtPolyRef ref, const float* pos, const float collisionQueryRange,
-		const float* forcedRemovePos, float forcedRemoveRadius, bool useForcedRemove,
+		const bool bRemoveNearLink, const float* linkV0, const float* linkV1,
 		const dtPolyRef* path, const int npath,
 		const float* moveDir,
 		dtNavMeshQuery* navquery, const dtQueryFilter* filter);
 
 	void update(const dtSharedBoundary* sharedData, const int sharedIdx,
 		const float* pos, const float collisionQueryRange,
-		const float* forcedRemovePos, float forcedRemoveRadius, bool useForcedRemove,
+		const bool bRemoveNearLink, const float* linkV0, const float* linkV1,
 		const dtPolyRef* path, const int npath, const float* moveDir,
 		dtNavMeshQuery* navquery, const dtQueryFilter* filter);
 	
diff --git a/Engine/Source/Runtime/Navmesh/Public/DetourCrowd/DetourSharedBoundary.h b/Engine/Source/Runtime/Navmesh/Public/DetourCrowd/DetourSharedBoundary.h
--- a/Engine/Source/Runtime/Navmesh/Public/DetourCrowd/DetourSharedBoundary.h
+++ b/Engine/Source/Runtime/Navmesh/Public/DetourCrowd/DetourSharedBoundary.h
@@ -16,12 +16,12 @@ struct dtSharedBoundaryData
 {
 	float Center[3];
 	float Radius;
+	float AccessTime;
 	dtQueryFilter* Filter;
-	int32 AccessTime;
 	uint8 SingleAreaId;
-
+	
 	TArray<dtSharedBoundaryEdge> Edges;
-	TArray<dtPolyRef> Polys;
+	TSet<dtPolyRef> Polys;
 
 	dtSharedBoundaryData() : Filter(nullptr) {}
 };
@@ -29,7 +29,7 @@ struct dtSharedBoundaryData
 class dtSharedBoundary
 {
 public:
-	TArray<dtSharedBoundaryData> Data;
+	TSparseArray<dtSharedBoundaryData> Data;
 	dtQueryFilter SingleAreaFilter;
 	float CurrentTime;
 	float NextClearTime;
@@ -37,11 +37,16 @@ public:
 	void Initialize();
 	void Tick(float DeltaTime);
 
-	int32 FindData(float* Center, float Radius, dtQueryFilter* NavFilter) const;
-	int32 FindData(float* Center, float Radius, uint8 SingleAreaId) const;
+	int32 FindData(float* Center, float Radius, dtPolyRef ReqPoly, dtQueryFilter* NavFilter) const;
+	int32 FindData(float* Center, float Radius, dtPolyRef ReqPoly, uint8 SingleAreaId) const;
 
 	int32 CacheData(float* Center, float Radius, dtPolyRef CenterPoly, dtNavMeshQuery* NavQuery, dtQueryFilter* NavFilter);
 	int32 CacheData(float* Center, float Radius, dtPolyRef CenterPoly, dtNavMeshQuery* NavQuery, uint8 SingleAreaId);
 
 	void FindEdges(dtSharedBoundaryData& Data, dtPolyRef CenterPoly, dtNavMeshQuery* NavQuery, dtQueryFilter* NavFilter);
+	bool HasSample(int32 Idx) const;
+
+private:
+
+	bool IsValid(int32 Idx, dtNavMeshQuery* NavQuery, dtQueryFilter* NavFilter) const;
 };
