fixed triangulation issue in navmesh's detailed polys
increased range of accurate (check against overlaps) tests in delaunay circumcircle conditions
#ue4 UE-11331

diff --git a/Engine/Source/Runtime/Navmesh/Private/DetourTileCache/DetourTileCacheDetail.cpp b/Engine/Source/Runtime/Navmesh/Private/DetourTileCache/DetourTileCacheDetail.cpp
--- a/Engine/Source/Runtime/Navmesh/Private/DetourTileCache/DetourTileCacheDetail.cpp
+++ b/Engine/Source/Runtime/Navmesh/Private/DetourTileCache/DetourTileCacheDetail.cpp
@@ -514,7 +514,9 @@ namespace TileCacheFunc
 					continue;
 				}
 				const float d = vdist2(c, &pts[u * 3]);
-				const float tol = 0.001f;
+				// UE4: increased tolerance of safe checks from 0.001f
+				// it was producing (rarely) overlapping edges
+				const float tol = 0.005f;
 				if (d > r*(1 + tol))
 				{
 					// Outside current circumcircle, skip.
