Added a way to query navigation data if a specified world location is inside a given navigation node #UE4

diff --git a/Engine/Source/Runtime/Engine/Classes/AI/Navigation/AbstractNavData.h b/Engine/Source/Runtime/Engine/Classes/AI/Navigation/AbstractNavData.h
--- a/Engine/Source/Runtime/Engine/Classes/AI/Navigation/AbstractNavData.h
+++ b/Engine/Source/Runtime/Engine/Classes/AI/Navigation/AbstractNavData.h
@@ -47,6 +47,7 @@ class AAbstractNavData : public ANavigationData
 	virtual ENavigationQueryResult::Type CalcPathCost(const FVector& PathStart, const FVector& PathEnd, float& OutPathCost, TSharedPtr<const FNavigationQueryFilter> QueryFilter = NULL, const UObject* Querier = NULL) const override { return ENavigationQueryResult::Invalid; }
 	virtual ENavigationQueryResult::Type CalcPathLength(const FVector& PathStart, const FVector& PathEnd, float& OutPathLength, TSharedPtr<const FNavigationQueryFilter> QueryFilter = NULL, const UObject* Querier = NULL) const override { return ENavigationQueryResult::Invalid; }
 	virtual ENavigationQueryResult::Type CalcPathLengthAndCost(const FVector& PathStart, const FVector& PathEnd, float& OutPathLength, float& OutPathCost, TSharedPtr<const FNavigationQueryFilter> QueryFilter = NULL, const UObject* Querier = NULL) const override { return ENavigationQueryResult::Invalid;	}
+	virtual bool DoesNodeContainLocation(NavNodeRef NodeRef, const FVector& WorldSpaceLocation) const override { return true; }
 	virtual void OnNavAreaAdded(const UClass* NavAreaClass, int32 AgentIndex) override {}
 	virtual void OnNavAreaRemoved(const UClass* NavAreaClass) override {};
 	// End ANavigationData overrides
diff --git a/Engine/Source/Runtime/Engine/Classes/AI/Navigation/NavigationData.h b/Engine/Source/Runtime/Engine/Classes/AI/Navigation/NavigationData.h
--- a/Engine/Source/Runtime/Engine/Classes/AI/Navigation/NavigationData.h
+++ b/Engine/Source/Runtime/Engine/Classes/AI/Navigation/NavigationData.h
@@ -676,7 +676,7 @@ public:
 	/** 
 	 *	Synchronously makes a raycast on navigation data using QueryFilter
 	 *	@param HitLocation if line was obstructed this will be set to hit location. Otherwise it contains SegmentEnd
-	 *	@return true if line from RayStart to RayEnd was obstructed
+	 *	@return true if line from RayStart to RayEnd is obstructed
 	 *
 	 *	@note don't make this function virtual! Look at implementation details and its comments for more info.
 	 */
@@ -718,6 +718,10 @@ public:
 	 *	@NOTE potentially expensive, so use it with caution */
 	virtual ENavigationQueryResult::Type CalcPathLengthAndCost(const FVector& PathStart, const FVector& PathEnd, float& OutPathLength, float& OutPathCost, TSharedPtr<const FNavigationQueryFilter> QueryFilter = NULL, const UObject* Querier = NULL) const PURE_VIRTUAL(ANavigationData::CalcPathLength, return ENavigationQueryResult::Invalid;);
 
+	/** Checks if specified navigation node contains given location 
+	 *	@param Location is expressed in WorldSpace, navigation data is responsible for tansforming if need be */
+	virtual bool DoesNodeContainLocation(NavNodeRef NodeRef, const FVector& WorldSpaceLocation) const PURE_VIRTUAL(ANavigationData::DoesNodeContainLocation, return false;);
+
 	const UWorld* GetCachedWorld() const;
 	float GetWorldTimeStamp() const;
 
diff --git a/Engine/Source/Runtime/Engine/Classes/AI/Navigation/RecastNavMesh.h b/Engine/Source/Runtime/Engine/Classes/AI/Navigation/RecastNavMesh.h
--- a/Engine/Source/Runtime/Engine/Classes/AI/Navigation/RecastNavMesh.h
+++ b/Engine/Source/Runtime/Engine/Classes/AI/Navigation/RecastNavMesh.h
@@ -686,6 +686,7 @@ public:
 	virtual ENavigationQueryResult::Type CalcPathCost(const FVector& PathStart, const FVector& PathEnd, float& OutPathCost, TSharedPtr<const FNavigationQueryFilter> Filter = NULL, const UObject* Querier = NULL) const override;
 	virtual ENavigationQueryResult::Type CalcPathLength(const FVector& PathStart, const FVector& PathEnd, float& OutPathLength, TSharedPtr<const FNavigationQueryFilter> QueryFilter = NULL, const UObject* Querier = NULL) const override;
 	virtual ENavigationQueryResult::Type CalcPathLengthAndCost(const FVector& PathStart, const FVector& PathEnd, float& OutPathLength, float& OutPathCost, TSharedPtr<const FNavigationQueryFilter> QueryFilter = NULL, const UObject* Querier = NULL) const override;
+	virtual bool DoesNodeContainLocation(NavNodeRef NodeRef, const FVector& WorldSpaceLocation) const override;
 
 	virtual UPrimitiveComponent* ConstructRenderingComponent() override;
 	/** Returns bounding box for the navmesh. */
diff --git a/Engine/Source/Runtime/Engine/Private/AI/Navigation/RecastNavMesh.cpp b/Engine/Source/Runtime/Engine/Private/AI/Navigation/RecastNavMesh.cpp
--- a/Engine/Source/Runtime/Engine/Private/AI/Navigation/RecastNavMesh.cpp
+++ b/Engine/Source/Runtime/Engine/Private/AI/Navigation/RecastNavMesh.cpp
@@ -1059,6 +1059,24 @@ ENavigationQueryResult::Type ARecastNavMesh::CalcPathLengthAndCost(const FVector
 	return Result;
 }
 
+bool ARecastNavMesh::DoesNodeContainLocation(NavNodeRef NodeRef, const FVector& WorldSpaceLocation) const
+{
+	bool bResult = false;
+	if (RecastNavMeshImpl != nullptr && RecastNavMeshImpl->GetRecastMesh() != nullptr)
+	{
+		dtNavMeshQuery NavQuery;
+		NavQuery.init(RecastNavMeshImpl->GetRecastMesh(), 0);
+
+		const FVector RcLocation = Unreal2RecastPoint(WorldSpaceLocation);
+		if (dtStatusFailed(NavQuery.isPointInsidePoly(NodeRef, &RcLocation.X, bResult)))
+		{
+			bResult = false;
+		}
+	}
+
+	return bResult; 
+}
+
 NavNodeRef ARecastNavMesh::FindNearestPoly(FVector const& Loc, FVector const& Extent, TSharedPtr<const FNavigationQueryFilter> Filter, const UObject* QueryOwner) const
 {
 	NavNodeRef PolyRef = 0;
@@ -1908,21 +1926,27 @@ bool ARecastNavMesh::SupportsStreaming() const
 }
 
 void ARecastNavMesh::ConditionalConstructGenerator()
-{
-	NavDataGenerator.Reset();
-	
+{	
 	UWorld* World = GetWorld();
 	check(World);
 	const bool bRequiresGenerator = SupportsRuntimeGeneration() || !World->IsGameWorld();
 	if (bRequiresGenerator)
 	{
-		NavDataGenerator = MakeShareable(new FRecastNavMeshGenerator(*this));
-
-		if (World->GetNavigationSystem())
+		if (NavDataGenerator.IsValid() == false)
 		{
-			RestrictBuildingToActiveTiles(World->GetNavigationSystem()->IsActiveTilesGenerationEnabled());
+			NavDataGenerator = MakeShareable(new FRecastNavMeshGenerator(*this));
+
+			if (World->GetNavigationSystem())
+			{
+				RestrictBuildingToActiveTiles(World->GetNavigationSystem()->IsActiveTilesGenerationEnabled());
+			}
 		}
 	}
+	else
+	{
+		NavDataGenerator->CancelBuild();
+		NavDataGenerator.Reset();
+	}
 }
 
 bool ARecastNavMesh::IsVoxelCacheEnabled()
diff --git a/Engine/Source/Runtime/Navmesh/Private/Detour/DetourNavMeshQuery.cpp b/Engine/Source/Runtime/Navmesh/Private/Detour/DetourNavMeshQuery.cpp
--- a/Engine/Source/Runtime/Navmesh/Private/Detour/DetourNavMeshQuery.cpp
+++ b/Engine/Source/Runtime/Navmesh/Private/Detour/DetourNavMeshQuery.cpp
@@ -225,50 +225,53 @@ dtStatus dtNavMeshQuery::init(const dtNavMesh* nav, const int maxNodes, const dt
 	m_nav = nav;
 	m_linkFilter = linkFilter;
 
-	if (!m_nodePool || m_nodePool->getMaxNodes() < maxNodes)
+	if (maxNodes > 0)
 	{
-		if (m_nodePool)
+		if (!m_nodePool || m_nodePool->getMaxNodes() < maxNodes)
 		{
-			m_nodePool->~dtNodePool();
-			dtFree(m_nodePool);
-			m_nodePool = 0;
+			if (m_nodePool)
+			{
+				m_nodePool->~dtNodePool();
+				dtFree(m_nodePool);
+				m_nodePool = 0;
+			}
+			m_nodePool = new (dtAlloc(sizeof(dtNodePool), DT_ALLOC_PERM)) dtNodePool(maxNodes, dtNextPow2(maxNodes / 4));
+			if (!m_nodePool)
+				return DT_FAILURE | DT_OUT_OF_MEMORY;
+		}
+		else
+		{
+			m_nodePool->clear();
 		}
-		m_nodePool = new (dtAlloc(sizeof(dtNodePool), DT_ALLOC_PERM)) dtNodePool(maxNodes, dtNextPow2(maxNodes/4));
-		if (!m_nodePool)
-			return DT_FAILURE | DT_OUT_OF_MEMORY;
-	}
-	else
-	{
-		m_nodePool->clear();
-	}
 
-	if (!m_tinyNodePool)
-	{
-		m_tinyNodePool = new (dtAlloc(sizeof(dtNodePool), DT_ALLOC_PERM)) dtNodePool(64, 32);
 		if (!m_tinyNodePool)
-			return DT_FAILURE | DT_OUT_OF_MEMORY;
-	}
-	else
-	{
-		m_tinyNodePool->clear();
-	}
-	
-	// TODO: check the open list size too.
-	if (!m_openList || m_openList->getCapacity() < maxNodes)
-	{
-		if (m_openList)
 		{
-			m_openList->~dtNodeQueue();
-			dtFree(m_openList);
-			m_openList = 0;
+			m_tinyNodePool = new (dtAlloc(sizeof(dtNodePool), DT_ALLOC_PERM)) dtNodePool(64, 32);
+			if (!m_tinyNodePool)
+				return DT_FAILURE | DT_OUT_OF_MEMORY;
+		}
+		else
+		{
+			m_tinyNodePool->clear();
+		}
+
+		// TODO: check the open list size too.
+		if (!m_openList || m_openList->getCapacity() < maxNodes)
+		{
+			if (m_openList)
+			{
+				m_openList->~dtNodeQueue();
+				dtFree(m_openList);
+				m_openList = 0;
+			}
+			m_openList = new (dtAlloc(sizeof(dtNodeQueue), DT_ALLOC_PERM)) dtNodeQueue(maxNodes);
+			if (!m_openList)
+				return DT_FAILURE | DT_OUT_OF_MEMORY;
+		}
+		else
+		{
+			m_openList->clear();
 		}
-		m_openList = new (dtAlloc(sizeof(dtNodeQueue), DT_ALLOC_PERM)) dtNodeQueue(maxNodes);
-		if (!m_openList)
-			return DT_FAILURE | DT_OUT_OF_MEMORY;
-	}
-	else
-	{
-		m_openList->clear();
 	}
 	
 	return DT_SUCCESS;
@@ -846,6 +849,31 @@ dtStatus dtNavMeshQuery::projectedPointOnPoly(dtPolyRef ref, const float* pos, f
 	return projectedPointOnPolyInTile(tile, poly, pos, projected);
 }
 
+dtStatus dtNavMeshQuery::isPointInsidePoly(dtPolyRef ref, const float* pos, bool& result) const
+{
+	dtAssert(m_nav);
+	const dtMeshTile* tile = 0;
+	const dtPoly* poly = 0;
+	if (dtStatusFailed(m_nav->getTileAndPolyByRef(ref, &tile, &poly)))
+		return DT_FAILURE | DT_INVALID_PARAM;
+	if (!tile)
+		return DT_FAILURE | DT_INVALID_PARAM;
+
+	if (poly->getType() == DT_POLYTYPE_OFFMESH_POINT)
+		return false;
+
+	const unsigned int ip = (unsigned int)(poly - tile->polys);
+
+	// Clamp point to be inside the polygon.
+	float verts[DT_VERTS_PER_POLYGON * 3];
+	const int nv = poly->vertCount;
+	for (int i = 0; i < nv; ++i)
+		dtVcopy(&verts[i * 3], &tile->verts[poly->verts[i] * 3]);
+
+	result = dtPointInPolygon(pos, verts, nv);
+
+	return DT_SUCCESS;
+}
 
 dtStatus dtNavMeshQuery::projectedPointOnPolyInTile(const dtMeshTile* tile, const dtPoly* poly,
 	const float* pos, float* projected) const
@@ -878,7 +906,7 @@ dtStatus dtNavMeshQuery::projectedPointOnPolyInTile(const dtMeshTile* tile, cons
 	if (dtPointInPolygon(pos, verts, nv))
 	{
 		// adjust point's height
-		// @todo this is an aproximation. Implement a proper solution if needed
+		// @todo this is an approximation. Implement a proper solution if needed
 		float h = 0;
 		for (int i = 0; i < nv; ++i)
 			h += verts[i * 3 + 1];
diff --git a/Engine/Source/Runtime/Navmesh/Public/Detour/DetourNavMeshQuery.h b/Engine/Source/Runtime/Navmesh/Public/Detour/DetourNavMeshQuery.h
--- a/Engine/Source/Runtime/Navmesh/Public/Detour/DetourNavMeshQuery.h
+++ b/Engine/Source/Runtime/Navmesh/Public/Detour/DetourNavMeshQuery.h
@@ -620,6 +620,13 @@ public:
 	/// @returns The status flags for the query.
 	dtStatus projectedPointOnPoly(dtPolyRef ref, const float* pos, float* projected) const;
 	
+	/// Checks if specified pos is inside given polygon specified by ref
+	///  @param[in]		ref			The reference id of the polygon.
+	///  @param[in]		pos			The position to check. [(x, y, z)]
+	///  @param[out]	result		The result of the check, whether the point is inside (true) or not (false)
+	/// @returns The status flags for the query.
+	dtStatus isPointInsidePoly(dtPolyRef ref, const float* pos, bool& result) const;
+
 	/// Gets the height of the polygon at the provided position using the height detail. (Most accurate.)
 	///  @param[in]		ref			The reference id of the polygon.
 	///  @param[in]		pos			A position within the xz-bounds of the polygon. [(x, y, z)]
