fixed cluster graph connections on navmesh
#ue4

diff --git a/Engine/Source/Runtime/Engine/Private/AI/Navigation/PImplRecastNavMesh.cpp b/Engine/Source/Runtime/Engine/Private/AI/Navigation/PImplRecastNavMesh.cpp
--- a/Engine/Source/Runtime/Engine/Private/AI/Navigation/PImplRecastNavMesh.cpp
+++ b/Engine/Source/Runtime/Engine/Private/AI/Navigation/PImplRecastNavMesh.cpp
@@ -2265,7 +2265,7 @@ int32 FPImplRecastNavMesh::GetTilesDebugGeometry(const FRecastNavMeshGenerator*
 		}
 	}
 
-	/*for (int32 i = 0; i < Header->clusterCount; i++)
+	for (int32 i = 0; i < Header->clusterCount; i++)
 	{
 		const dtCluster& c0 = Tile.clusters[i];
 		uint32 iLink = c0.firstLink;
@@ -2296,7 +2296,7 @@ int32 FPImplRecastNavMesh::GetTilesDebugGeometry(const FRecastNavMeshGenerator*
 				OutGeometry.ClusterLinks.Add(LinkGeom);
 			}
 		}
-	}*/
+	}
 
 	// Get tile edges and navmesh edges
 	if (OutGeometry.bGatherPolyEdges || OutGeometry.bGatherNavMeshEdges)
diff --git a/Engine/Source/Runtime/Navmesh/Private/DetourTileCache/DetourTileCacheBuilder.cpp b/Engine/Source/Runtime/Navmesh/Private/DetourTileCache/DetourTileCacheBuilder.cpp
--- a/Engine/Source/Runtime/Navmesh/Private/DetourTileCache/DetourTileCacheBuilder.cpp
+++ b/Engine/Source/Runtime/Navmesh/Private/DetourTileCache/DetourTileCacheBuilder.cpp
@@ -885,7 +885,7 @@ dtStatus dtBuildTileCacheContours(dtTileCacheAlloc* alloc, dtTileCacheLayer& lay
 	}
 
 	// Build clusters
-	clusters.nregs = layer.regCount;
+	clusters.nregs = layer.regCount ? (layer.regCount + 1) : 0;
 	clusters.npolys = 0;
 	clusters.nclusters = 0;
 	clusters.regMap = (unsigned short*)alloc->alloc(sizeof(unsigned short)*clusters.nregs);
@@ -904,8 +904,8 @@ dtStatus dtBuildTileCacheContours(dtTileCacheAlloc* alloc, dtTileCacheLayer& lay
 	// - gather all neighbor regions
 	// - repeat inner loop for every region from list
 	
-	dtScopedDelete<unsigned short> neiRegs(layer.regCount);
-	dtScopedDelete<unsigned short> newNeiRegs(layer.regCount);
+	dtScopedDelete<unsigned short> neiRegs(layer.regCount + 1);
+	dtScopedDelete<unsigned short> newNeiRegs(layer.regCount + 1);
 	int nneiRegs = 0;
 	int nnewNeiRegs = 0;
 	
@@ -924,7 +924,7 @@ dtStatus dtBuildTileCacheContours(dtTileCacheAlloc* alloc, dtTileCacheLayer& lay
 		{
 			// there could be more than one contour per region...
 			dtTileCacheContour& cont = lcset.conts[ic];
-			if (cont.reg != (unsigned short)i || cont.area == DT_TILECACHE_NULL_AREA)
+			if (cont.reg != (unsigned short)(i) || cont.area == DT_TILECACHE_NULL_AREA)
 			{
 				continue;
 			}
