navigation links will connect to cheapest area in snap radius
#ue4, merged from dev branch

diff --git a/Engine/Source/Runtime/Engine/Classes/AI/Navigation/NavLinkDefinition.h b/Engine/Source/Runtime/Engine/Classes/AI/Navigation/NavLinkDefinition.h
--- a/Engine/Source/Runtime/Engine/Classes/AI/Navigation/NavLinkDefinition.h
+++ b/Engine/Source/Runtime/Engine/Classes/AI/Navigation/NavLinkDefinition.h
@@ -40,6 +40,10 @@ struct ENGINE_API FNavigationLinkBase
 	UPROPERTY()
 	uint32 bUseSnapHeight : 1;
 
+	/** If set, link will try to snap to cheapest area in given radius */
+	UPROPERTY(EditAnywhere, Category=Default)
+	uint32 bSnapToCheapestArea : 1;
+
 	/** Area type of this link (empty = default) */
 	UPROPERTY(EditAnywhere, Category=Default)
 	TSubclassOf<class UNavArea> AreaClass;
diff --git a/Engine/Source/Runtime/Engine/Classes/AI/Navigation/NavigationData.h b/Engine/Source/Runtime/Engine/Classes/AI/Navigation/NavigationData.h
--- a/Engine/Source/Runtime/Engine/Classes/AI/Navigation/NavigationData.h
+++ b/Engine/Source/Runtime/Engine/Classes/AI/Navigation/NavigationData.h
@@ -734,7 +734,10 @@ public:
 	
 	/** area was removed from navigation system */
 	virtual void OnNavAreaRemoved(const UClass* NavAreaClass);
-		
+
+	/** called after changes to registered area classes */
+	virtual void OnNavAreaChanged();
+
 	void OnNavAreaEvent(const UClass* NavAreaClass, ENavAreaEvent::Type Event);
 
 	/** add all existing areas */
diff --git a/Engine/Source/Runtime/Engine/Classes/AI/Navigation/RecastNavMesh.h b/Engine/Source/Runtime/Engine/Classes/AI/Navigation/RecastNavMesh.h
--- a/Engine/Source/Runtime/Engine/Classes/AI/Navigation/RecastNavMesh.h
+++ b/Engine/Source/Runtime/Engine/Classes/AI/Navigation/RecastNavMesh.h
@@ -783,10 +783,11 @@ public:
 	//----------------------------------------------------------------------//
 	// Utilities
 	//----------------------------------------------------------------------//
+	virtual void OnNavAreaChanged() override;
 	virtual void OnNavAreaAdded(const UClass* NavAreaClass, int32 AgentIndex) override;
 	virtual int32 GetNewAreaID(const UClass* AreaClass) const override;
 	virtual int32 GetMaxSupportedAreas() const override { return RECAST_MAX_AREAS; }
-	
+
 	/** Get forbidden area flags from default query filter */
 	uint16 GetDefaultForbiddenFlags() const;
 	/** Change forbidden area flags in default query filter */
diff --git a/Engine/Source/Runtime/Engine/Private/AI/Navigation/NavigationData.cpp b/Engine/Source/Runtime/Engine/Private/AI/Navigation/NavigationData.cpp
--- a/Engine/Source/Runtime/Engine/Private/AI/Navigation/NavigationData.cpp
+++ b/Engine/Source/Runtime/Engine/Private/AI/Navigation/NavigationData.cpp
@@ -579,6 +579,8 @@ void ANavigationData::OnNavAreaEvent(const UClass* NavAreaClass, ENavAreaEvent::
 	{
 		OnNavAreaRemoved(NavAreaClass);
 	}
+
+	OnNavAreaChanged();
 }
 
 void ANavigationData::OnNavAreaRemoved(const UClass* NavAreaClass)
@@ -594,12 +596,19 @@ void ANavigationData::OnNavAreaRemoved(const UClass* NavAreaClass)
 	}
 }
 
+void ANavigationData::OnNavAreaChanged()
+{
+	// empty in base class
+}
+
 void ANavigationData::ProcessNavAreas(const TArray<const UClass*>& AreaClasses, int32 AgentIndex)
 {
 	for (int32 i = 0; i < AreaClasses.Num(); i++)
 	{
 		OnNavAreaAdded(AreaClasses[i], AgentIndex);
 	}
+
+	OnNavAreaChanged();
 }
 
 int32 ANavigationData::GetNewAreaID(const UClass* AreaClass) const
diff --git a/Engine/Source/Runtime/Engine/Private/AI/Navigation/NavigationModifier.cpp b/Engine/Source/Runtime/Engine/Private/AI/Navigation/NavigationModifier.cpp
--- a/Engine/Source/Runtime/Engine/Private/AI/Navigation/NavigationModifier.cpp
+++ b/Engine/Source/Runtime/Engine/Private/AI/Navigation/NavigationModifier.cpp
@@ -24,7 +24,7 @@ static const float CONVEX_HULL_POINTS_MIN_DISTANCE_SQ = 4.0f * 4.0f;
 //----------------------------------------------------------------------//
 FNavigationLinkBase::FNavigationLinkBase() 
 	: MaxFallDownLength(1000.0f), Direction(ENavLinkDirection::BothWays), UserId(0),
-	  SnapRadius(30.f), SnapHeight(50.0f), bUseSnapHeight(false)
+	  SnapRadius(30.f), SnapHeight(50.0f), bUseSnapHeight(false), bSnapToCheapestArea(true)
 {
 	AreaClass = NULL;
 	SupportedAgentsBits = 0xFFFFFFFF;
diff --git a/Engine/Source/Runtime/Engine/Private/AI/Navigation/PImplRecastNavMesh.cpp b/Engine/Source/Runtime/Engine/Private/AI/Navigation/PImplRecastNavMesh.cpp
--- a/Engine/Source/Runtime/Engine/Private/AI/Navigation/PImplRecastNavMesh.cpp
+++ b/Engine/Source/Runtime/Engine/Private/AI/Navigation/PImplRecastNavMesh.cpp
@@ -2521,6 +2521,43 @@ void FPImplRecastNavMesh::SetFilterForbiddenFlags(FRecastQueryFilter* Filter, ui
 	// include-exclude don't need to be symmetrical, filter will check both conditions
 }
 
+void FPImplRecastNavMesh::OnAreaCostChanged()
+{
+	struct FFloatIntPair
+	{
+		float Score;
+		int32 Index;
+
+		FFloatIntPair() : Score(MAX_FLT), Index(0) {}
+		FFloatIntPair(int32 AreaId, float TravelCost, float EntryCost) : Score(TravelCost + EntryCost), Index(AreaId) {}
+
+		bool operator <(const FFloatIntPair& Other) const { return Score < Other.Score; }
+	};
+
+	if (NavMeshOwner && DetourNavMesh)
+	{
+		const INavigationQueryFilterInterface* NavFilter = NavMeshOwner->GetDefaultQueryFilterImpl();
+		const dtQueryFilter* DetourFilter = ((const FRecastQueryFilter*)NavFilter)->GetAsDetourQueryFilter();
+
+		TArray<FFloatIntPair> AreaData;
+		AreaData.Reserve(RECAST_MAX_AREAS);
+		for (int32 Idx = 0; Idx < RECAST_MAX_AREAS; Idx++)
+		{
+			AreaData.Add(FFloatIntPair(Idx, DetourFilter->getAreaCost(Idx), DetourFilter->getAreaFixedCost(Idx)));
+		}
+
+		AreaData.Sort();
+
+		uint8 AreaCostOrder[RECAST_MAX_AREAS];
+		for (int32 Idx = 0; Idx < RECAST_MAX_AREAS; Idx++)
+		{
+			AreaCostOrder[AreaData[Idx].Index] = Idx;
+		}
+
+		DetourNavMesh->applyAreaCostOrder(AreaCostOrder);
+	}
+}
+
 void FPImplRecastNavMesh::RemoveTileCacheLayers(int32 TileX, int32 TileY)
 {
 	CompressedTileCacheLayers.Remove(FIntPoint(TileX, TileY));
diff --git a/Engine/Source/Runtime/Engine/Private/AI/Navigation/RecastNavMesh.cpp b/Engine/Source/Runtime/Engine/Private/AI/Navigation/RecastNavMesh.cpp
--- a/Engine/Source/Runtime/Engine/Private/AI/Navigation/RecastNavMesh.cpp
+++ b/Engine/Source/Runtime/Engine/Private/AI/Navigation/RecastNavMesh.cpp
@@ -586,6 +586,14 @@ void ARecastNavMesh::OnNavAreaAdded(const UClass* NavAreaClass, int32 AgentIndex
 	}
 }
 
+void ARecastNavMesh::OnNavAreaChanged()
+{
+	if (RecastNavMeshImpl)
+	{
+		RecastNavMeshImpl->OnAreaCostChanged();
+	}
+}
+
 int32 ARecastNavMesh::GetNewAreaID(const UClass* AreaClass) const
 {
 	if (AreaClass == UNavigationSystem::GetDefaultWalkableArea())
diff --git a/Engine/Source/Runtime/Engine/Private/AI/Navigation/RecastNavMeshGenerator.cpp b/Engine/Source/Runtime/Engine/Private/AI/Navigation/RecastNavMeshGenerator.cpp
--- a/Engine/Source/Runtime/Engine/Private/AI/Navigation/RecastNavMeshGenerator.cpp
+++ b/Engine/Source/Runtime/Engine/Private/AI/Navigation/RecastNavMeshGenerator.cpp
@@ -1271,7 +1271,10 @@ struct FOffMeshData
 			StoreUnrealPoint(NewInfo.vertsA0, LocalToWorld.TransformPosition(Link.Left));
 			StoreUnrealPoint(NewInfo.vertsB0, LocalToWorld.TransformPosition(Link.Right));
 
-			NewInfo.type = DT_OFFMESH_CON_POINT | (Link.Direction == ENavLinkDirection::BothWays ? DT_OFFMESH_CON_BIDIR : 0);
+			NewInfo.type = DT_OFFMESH_CON_POINT | 
+				(Link.Direction == ENavLinkDirection::BothWays ? DT_OFFMESH_CON_BIDIR : 0) |
+				(Link.bSnapToCheapestArea ? DT_OFFMESH_CON_CHEAPAREA : 0);
+
 			NewInfo.snapRadius = Link.SnapRadius;
 			NewInfo.snapHeight = Link.bUseSnapHeight ? Link.SnapHeight : DefaultSnapHeight;
 			NewInfo.userID = Link.UserId;
diff --git a/Engine/Source/Runtime/Engine/Public/AI/Navigation/PImplRecastNavMesh.h b/Engine/Source/Runtime/Engine/Public/AI/Navigation/PImplRecastNavMesh.h
--- a/Engine/Source/Runtime/Engine/Public/AI/Navigation/PImplRecastNavMesh.h
+++ b/Engine/Source/Runtime/Engine/Public/AI/Navigation/PImplRecastNavMesh.h
@@ -190,6 +190,8 @@ public:
 	static uint16 GetFilterForbiddenFlags(const FRecastQueryFilter* Filter);
 	static void SetFilterForbiddenFlags(FRecastQueryFilter* Filter, uint16 ForbiddenFlags);
 
+	void OnAreaCostChanged();
+
 public:
 	dtNavMesh const* GetRecastMesh() const { return DetourNavMesh; };
 	dtNavMesh* GetRecastMesh() { return DetourNavMesh; };
diff --git a/Engine/Source/Runtime/Navmesh/Private/Detour/DetourNavMesh.cpp b/Engine/Source/Runtime/Navmesh/Private/Detour/DetourNavMesh.cpp
--- a/Engine/Source/Runtime/Navmesh/Private/Detour/DetourNavMesh.cpp
+++ b/Engine/Source/Runtime/Navmesh/Private/Detour/DetourNavMesh.cpp
@@ -734,6 +734,7 @@ dtNavMesh::dtNavMesh() :
 	m_polyBits(0)
 {
 	memset(&m_params, 0, sizeof(dtNavMeshParams));
+	memset(m_areaCostOrder, 0, sizeof(m_areaCostOrder));
 	m_orig[0] = 0;
 	m_orig[1] = 0;
 	m_orig[2] = 0;
@@ -1073,7 +1074,10 @@ void dtNavMesh::connectExtOffMeshLinks(dtMeshTile* tile, dtMeshTile* target, int
 		// Find polygon to connect to.
 		const float* p = &targetCon->pos[3];
 		float nearestPt[3];
-		dtPolyRef ref = findNearestPolyInTile(tile, p, ext, nearestPt, true);
+		dtPolyRef ref = targetCon->getSnapToCheapestArea() ?
+			findCheapestNearPolyInTile(tile, p, ext, nearestPt) :
+			findNearestPolyInTile(tile, p, ext, nearestPt, true);
+
 		if (!ref)
 			continue;
 		// findNearestPoly may return too optimistic results, further check to make sure. 
@@ -1187,7 +1191,9 @@ void dtNavMesh::baseOffMeshLinks(dtMeshTile* tile)
 		// Find polygon to connect to.
 		const float* p = &con->pos[0]; // First vertex
 		float nearestPt[3];
-		dtPolyRef ref = findNearestPolyInTile(tile, p, ext, nearestPt, true);
+		dtPolyRef ref = con->getSnapToCheapestArea() ? 
+			findCheapestNearPolyInTile(tile, p, ext, nearestPt) :
+			findNearestPolyInTile(tile, p, ext, nearestPt, true);
 		if (!ref) continue;
 		// findNearestPoly may return too optimistic results, further check to make sure. 
 		if (dtSqr(nearestPt[0]-p[0])+dtSqr(nearestPt[2]-p[2]) > dtSqr(con->rad))
@@ -1409,6 +1415,60 @@ dtPolyRef dtNavMesh::findNearestPolyInTile(const dtMeshTile* tile,
 	return nearest;
 }
 
+dtPolyRef dtNavMesh::findCheapestNearPolyInTile(const dtMeshTile* tile, const float* center,
+												const float* extents, float* nearestPt) const
+{
+	float bmin[3], bmax[3];
+	dtVsub(bmin, center, extents);
+	dtVadd(bmax, center, extents);
+
+	// Get nearby polygons from proximity grid.
+	dtPolyRef polys[128];
+	const bool bExcludeUnwalkable = true;
+	int polyCount = queryPolygonsInTile(tile, bmin, bmax, polys, 128, bExcludeUnwalkable);
+
+	// Find nearest polygon amongst the nearby polygons.
+	dtPolyRef nearest = 0;
+	float nearestDistanceSqr = FLT_MAX;
+	unsigned char cheapestAreaCostOrder = 0xff;
+	for (int i = 0; i < polyCount; ++i)
+	{
+		dtPolyRef ref = polys[i];
+		
+		const int polyIdx = decodePolyIdPoly(ref);
+		dtPoly* poly = &tile->polys[polyIdx];
+		const unsigned char polyAreaCostOrder = m_areaCostOrder[poly->getArea()];
+		if (polyAreaCostOrder < cheapestAreaCostOrder)
+		{
+			cheapestAreaCostOrder = polyAreaCostOrder;
+			nearestDistanceSqr = FLT_MAX;
+			nearest = 0;
+		}
+
+		if (polyAreaCostOrder == cheapestAreaCostOrder)
+		{
+			float closestPtPoly[3];
+			closestPointOnPolyInTile(tile, polyIdx, center, closestPtPoly);
+			float d = dtVdistSqr(center, closestPtPoly);
+			if (d < nearestDistanceSqr)
+			{
+				if (nearestPt)
+					dtVcopy(nearestPt, closestPtPoly);
+				nearestDistanceSqr = d;
+				nearest = ref;
+			}
+		}
+	}
+
+	// Verify if the point is actually within requested height, caller is performing 2D check anyway (radius)
+	if (dtAbs(nearestPt[1] - center[1]) > extents[1])
+	{
+		nearest = 0;
+	}
+
+	return nearest;
+}
+
 int dtNavMesh::queryPolygonsInTile(const dtMeshTile* tile, const float* qmin, const float* qmax,
 									dtPolyRef* polys, const int maxPolys, bool bExcludeUnwalkable) const
 {
@@ -2407,4 +2467,9 @@ void dtNavMesh::applyWorldOffset(const float* offset)
 		}
 	}
 }
+
+void dtNavMesh::applyAreaCostOrder(unsigned char* costOrder)
+{
+	memcpy(m_areaCostOrder, costOrder, sizeof(m_areaCostOrder));
+}
 //@UE4 END
diff --git a/Engine/Source/Runtime/Navmesh/Public/Detour/DetourNavMesh.h b/Engine/Source/Runtime/Navmesh/Public/Detour/DetourNavMesh.h
--- a/Engine/Source/Runtime/Navmesh/Public/Detour/DetourNavMesh.h
+++ b/Engine/Source/Runtime/Navmesh/Public/Detour/DetourNavMesh.h
@@ -103,6 +103,7 @@ static const unsigned int DT_NULL_LINK = 0xffffffff;
 static const unsigned char DT_OFFMESH_CON_BIDIR = 0x01;
 static const unsigned char DT_OFFMESH_CON_POINT = 0x02;
 static const unsigned char DT_OFFMESH_CON_SEGMENT = 0x04;
+static const unsigned char DT_OFFMESH_CON_CHEAPAREA = 0x08;
 
 /// The maximum number of user defined area ids.
 /// @ingroup detour
@@ -299,11 +300,15 @@ struct dtOffMeshConnection
 	/// Sets link flags
 	inline void setFlags(unsigned char conFlags)
 	{
-		flags = ((conFlags & DT_OFFMESH_CON_BIDIR) ? 0x80 : 0); 
+		flags = ((conFlags & DT_OFFMESH_CON_BIDIR) ? 0x80 : 0) |
+			((conFlags & DT_OFFMESH_CON_CHEAPAREA) ? 0x40 : 0);
 	}
 
 	/// Gets the link direction
 	inline bool getBiDirectional() const { return (flags & 0x80) != 0; }
+
+	/// Gets the link snap mode
+	inline bool getSnapToCheapestArea() const { return (flags & 0x40) != 0; }
 };
 
 /// Cluster of polys
@@ -738,6 +743,8 @@ public:
 		// has no tile grid set up
 		return (m_tileWidth > 0 && m_tileHeight > 0) == false;
 	}
+
+	void applyAreaCostOrder(unsigned char* costOrder);
 	//@UE4 END
 	
 private:
@@ -792,6 +799,8 @@ private:
 	/// Find nearest polygon within a tile.
 	dtPolyRef findNearestPolyInTile(const dtMeshTile* tile, const float* center,
 									const float* extents, float* nearestPt, bool bExcludeUnwalkable = false) const;
+	dtPolyRef findCheapestNearPolyInTile(const dtMeshTile* tile, const float* center,
+										 const float* extents, float* nearestPt) const;
 	/// Returns closest point on polygon.
 	void closestPointOnPolyInTile(const dtMeshTile* tile, unsigned int ip,
 								  const float* pos, float* closest) const;
@@ -803,6 +812,8 @@ private:
 	int m_tileLutSize;					///< Tile hash lookup size (must be pot).
 	int m_tileLutMask;					///< Tile hash lookup mask.
 
+	unsigned char m_areaCostOrder[DT_MAX_AREAS];
+
 	dtMeshTile** m_posLookup;			///< Tile hash lookup.
 	dtMeshTile* m_nextFree;				///< Freelist of tiles.
 	dtMeshTile* m_tiles;				///< List of tiles.
