Fixed variable shadowing in Recast code #UE4

diff --git a/Engine/Source/Runtime/Navmesh/Private/Detour/DetourNavMeshQuery.cpp b/Engine/Source/Runtime/Navmesh/Private/Detour/DetourNavMeshQuery.cpp
--- a/Engine/Source/Runtime/Navmesh/Private/Detour/DetourNavMeshQuery.cpp
+++ b/Engine/Source/Runtime/Navmesh/Private/Detour/DetourNavMeshQuery.cpp
@@ -87,10 +87,9 @@ void dtQueryFilterData::copyFrom(const dtQueryFilterData* source)
 {
 	memcpy((void*)this, source, sizeof(dtQueryFilterData));
 }
-	
+
 //@UE4 BEGIN
-// removed following line to make H_SCALE parametrizable (via dtQueryFilter::heuristicScale)
-static const float H_SCALE = 0.999f; // Search heuristic scale.
+static const float DEFAULT_HEURISTIC_SCALE = 0.999f; // Search heuristic scale.
 //@UE4 END
 
 dtNavMeshQuery* dtAllocNavMeshQuery()
@@ -1548,7 +1547,7 @@ dtStatus dtNavMeshQuery::testClusterPath(dtPolyRef startRef, dtPolyRef endRef) c
 	dtVcopy(startNode->pos, startCluster.center);
 	startNode->pidx = 0;
 	startNode->cost = 0;
-	startNode->total = dtVdist(startCluster.center, endCluster.center) * H_SCALE;
+	startNode->total = dtVdist(startCluster.center, endCluster.center) * DEFAULT_HEURISTIC_SCALE;
 	startNode->id = startCRef;
 	startNode->flags = DT_NODE_OPEN;
 	m_openList->push(startNode);
@@ -1619,7 +1618,7 @@ dtStatus dtNavMeshQuery::testClusterPath(dtPolyRef startRef, dtPolyRef endRef) c
 
 			// Calculate cost and heuristic.
 			const float cost = bestNode->cost;
-			const float heuristic = (neighbourRef != endCRef) ? dtVdist(neighbourNode->pos, endCluster.center)*H_SCALE : 0.0f;
+			const float heuristic = (neighbourRef != endCRef) ? dtVdist(neighbourNode->pos, endCluster.center)*DEFAULT_HEURISTIC_SCALE : 0.0f;
 			const float total = cost + heuristic;
 
 			// The node is already in open list and the new result is worse, skip.
@@ -2654,11 +2653,11 @@ dtStatus dtNavMeshQuery::getPortalPoints(dtPolyRef from, const dtPoly* fromPoly,
 {
 	// Find the link that points to the 'to' polygon.
 	const dtLink* link = 0;
-	unsigned int i = fromPoly->firstLink;
-	while (i != DT_NULL_LINK)
+	unsigned int linkIndex = fromPoly->firstLink;
+	while (linkIndex != DT_NULL_LINK)
 	{
-		const dtLink& testLink = m_nav->getLink(fromTile, i);
-		i = testLink.next;
+		const dtLink& testLink = m_nav->getLink(fromTile, linkIndex);
+		linkIndex = testLink.next;
 		
 		if (testLink.ref == to)
 		{
@@ -3579,7 +3578,7 @@ static bool containsPolyRef(const dtPolyRef testRef, const dtPolyRef* path, cons
 /// the start polygon.
 ///
 /// The same intersection test restrictions that apply to the findPolysAroundCircle 
-/// mehtod applies to this method.
+/// method applies to this method.
 ///
 /// The value of the center point is used as the start point for cost calculations. 
 /// It is not projected onto the surface of the mesh, so its y-value will effect 
@@ -3642,8 +3641,8 @@ dtStatus dtNavMeshQuery::findLocalNeighbourhood(dtPolyRef startRef, const float*
 	{
 		// Pop front.
 		dtNode* curNode = stack[0];
-		for (int i = 0; i < nstack-1; ++i)
-			stack[i] = stack[i+1];
+		for (int stackIndex = 0; stackIndex < nstack - 1; ++stackIndex)
+			stack[stackIndex] = stack[stackIndex + 1];
 		nstack--;
 		
 		// Get poly and tile.
@@ -3707,8 +3706,8 @@ dtStatus dtNavMeshQuery::findLocalNeighbourhood(dtPolyRef startRef, const float*
 			
 			// Collect vertices of the neighbour poly.
 			const int npa = neighbourPoly->vertCount;
-			for (int k = 0; k < npa; ++k)
-				dtVcopy(&pa[k*3], &neighbourTile->verts[neighbourPoly->verts[k]*3]);
+			for (int neighbourPolyVertIndex = 0; neighbourPolyVertIndex < npa; ++neighbourPolyVertIndex)
+				dtVcopy(&pa[neighbourPolyVertIndex * 3], &neighbourTile->verts[neighbourPoly->verts[neighbourPolyVertIndex] * 3]);
 			
 			bool overlap = false;
 			for (int j = 0; j < n; ++j)
@@ -3717,11 +3716,11 @@ dtStatus dtNavMeshQuery::findLocalNeighbourhood(dtPolyRef startRef, const float*
 				
 				// Connected polys do not overlap.
 				bool connected = false;
-				unsigned int k = neighbourPoly->firstLink;
-				while (k != DT_NULL_LINK)
+				unsigned int neighbourLinkIndex = neighbourPoly->firstLink;
+				while (neighbourLinkIndex != DT_NULL_LINK)
 				{
-					const dtLink& link2 = m_nav->getLink(neighbourTile, k);
-					k = link2.next;
+					const dtLink& link2 = m_nav->getLink(neighbourTile, neighbourLinkIndex);
+					neighbourLinkIndex = link2.next;
 
 					if (link2.ref == pastRef)
 					{
@@ -3739,8 +3738,8 @@ dtStatus dtNavMeshQuery::findLocalNeighbourhood(dtPolyRef startRef, const float*
 				
 				// Get vertices and test overlap
 				const int npb = pastPoly->vertCount;
-				for (int k = 0; k < npb; ++k)
-					dtVcopy(&pb[k*3], &pastTile->verts[pastPoly->verts[k]*3]);
+				for (int pastPolyVertIndex = 0; pastPolyVertIndex < npb; ++pastPolyVertIndex)
+					dtVcopy(&pb[pastPolyVertIndex * 3], &pastTile->verts[pastPoly->verts[pastPolyVertIndex] * 3]);
 				
 				if (dtOverlapPolyPoly2D(pa,npa, pb,npb))
 				{
@@ -4026,8 +4025,8 @@ dtStatus dtNavMeshQuery::findWallsInNeighbourhood(dtPolyRef startRef, const floa
 	{
 		// Pop front.
 		dtNode* curNode = stack[0];
-		for (int i = 0; i < nstack - 1; ++i)
-			stack[i] = stack[i + 1];
+		for (int stackIndex = 0; stackIndex < nstack - 1; ++stackIndex)
+			stack[stackIndex] = stack[stackIndex + 1];
 		nstack--;
 		
 		// Get poly and tile.
@@ -4115,10 +4114,10 @@ dtStatus dtNavMeshQuery::findWallsInNeighbourhood(dtPolyRef startRef, const floa
 		}
 
 		// add hard edges of poly
-		for (int i = 0; i < curPoly->vertCount; i++)
+		for (int neighbourIndex = 0; neighbourIndex < curPoly->vertCount; neighbourIndex++)
 		{
-			bool bStoreEdge = (curPoly->neis[i] == 0);
-			if (curPoly->neis[i] & DT_EXT_LINK)
+			bool bStoreEdge = (curPoly->neis[neighbourIndex] == 0);
+			if (curPoly->neis[neighbourIndex] & DT_EXT_LINK)
 			{
 				// check if external edge has valid link
 				bool bConnected = false;
@@ -4129,7 +4128,7 @@ dtStatus dtNavMeshQuery::findWallsInNeighbourhood(dtPolyRef startRef, const floa
 					const dtLink& link = m_nav->getLink(curTile, linkIdx);
 					linkIdx = link.next;
 
-					if (link.edge == i)
+					if (link.edge == neighbourIndex)
 					{
 						bConnected = true;
 						break;
@@ -4141,7 +4140,7 @@ dtStatus dtNavMeshQuery::findWallsInNeighbourhood(dtPolyRef startRef, const floa
 
 			if (bStoreEdge)
 			{
-				storeWallSegment(curTile, curPoly, i, 
+				storeWallSegment(curTile, curPoly, neighbourIndex,
 					curRef, 0, centerPos, radiusSqr, 
 					resultWalls, resultRefs, resultCount, maxResult);
 			}
diff --git a/Engine/Source/Runtime/Navmesh/Private/Recast/RecastMesh.cpp b/Engine/Source/Runtime/Navmesh/Private/Recast/RecastMesh.cpp
--- a/Engine/Source/Runtime/Navmesh/Private/Recast/RecastMesh.cpp
+++ b/Engine/Source/Runtime/Navmesh/Private/Recast/RecastMesh.cpp
@@ -774,10 +774,10 @@ static bool removeVertex(rcContext* ctx, rcPolyMesh& mesh, const unsigned short
 	// Generate temp vertex array for triangulation.
 	for (int i = 0; i < nhole; ++i)
 	{
-		const int pi = hole[i];
-		tverts[i*4+0] = mesh.verts[pi*3+0];
-		tverts[i*4+1] = mesh.verts[pi*3+1];
-		tverts[i*4+2] = mesh.verts[pi*3+2];
+		const int holeVertIndex = hole[i];
+		tverts[i*4+0] = mesh.verts[holeVertIndex*3+0];
+		tverts[i*4+1] = mesh.verts[holeVertIndex*3+1];
+		tverts[i*4+2] = mesh.verts[holeVertIndex*3+2];
 		tverts[i*4+3] = 0;
 		thole[i] = i;
 	}
