Wrap remaining shadow variables with disable pragma until owners fix

diff --git a/Engine/Source/Developer/MeshSimplifier/Private/MeshSimplify.h b/Engine/Source/Developer/MeshSimplifier/Private/MeshSimplify.h
--- a/Engine/Source/Developer/MeshSimplifier/Private/MeshSimplify.h
+++ b/Engine/Source/Developer/MeshSimplifier/Private/MeshSimplify.h
@@ -1244,6 +1244,7 @@ void TMeshSimplifier<T, NumAttributes>::Collapse( TSimpEdge<T>* edge )
 	numVerts--;
 }
 
+PRAGMA_DISABLE_SHADOW_VARIABLE_WARNINGS
 template< typename T, uint32 NumAttributes >
 void TMeshSimplifier<T, NumAttributes>::SimplifyMesh( float maxError, int minTris )
 {
@@ -1634,6 +1635,7 @@ void TMeshSimplifier<T, NumAttributes>::SimplifyMesh( float maxError, int minTri
 
 	//common->Printf( "simplified numVerts %i, numTris %i\n", numVerts, numTris );
 }
+PRAGMA_POP
 
 template< typename T, uint32 NumAttributes >
 void TMeshSimplifier<T, NumAttributes>::OutputMesh( T* verts, uint32* indexes ) {
diff --git a/Engine/Source/Runtime/Navmesh/Private/DetourTileCache/DetourTileCacheBuilder.cpp b/Engine/Source/Runtime/Navmesh/Private/DetourTileCache/DetourTileCacheBuilder.cpp
--- a/Engine/Source/Runtime/Navmesh/Private/DetourTileCache/DetourTileCacheBuilder.cpp
+++ b/Engine/Source/Runtime/Navmesh/Private/DetourTileCache/DetourTileCacheBuilder.cpp
@@ -1330,7 +1330,9 @@ static bool diagonalie(int i, int j, int n, const unsigned short* verts, const u
 // polygon P in the neighborhood of the i endpoint.
 static bool	inCone(int i, int j, int n, const unsigned short* verts, const unsigned short* indices)
 {
+	PRAGMA_DISABLE_SHADOW_VARIABLE_WARNINGS
 	const unsigned short* pi = &verts[(indices[i] & 0x7fff) * 4];
+	PRAGMA_POP
 	const unsigned short* pj = &verts[(indices[j] & 0x7fff) * 4];
 	const unsigned short* pi1 = &verts[(indices[TileCacheFunc::next(i, n)] & 0x7fff) * 4];
 	const unsigned short* pin1 = &verts[(indices[TileCacheFunc::prev(i, n)] & 0x7fff) * 4];
diff --git a/Engine/Source/Runtime/Navmesh/Private/Recast/RecastMesh.cpp b/Engine/Source/Runtime/Navmesh/Private/Recast/RecastMesh.cpp
--- a/Engine/Source/Runtime/Navmesh/Private/Recast/RecastMesh.cpp
+++ b/Engine/Source/Runtime/Navmesh/Private/Recast/RecastMesh.cpp
@@ -272,7 +272,9 @@ static bool diagonalie(int i, int j, int n, const int* verts, int* indices)
 // polygon P in the neighborhood of the i endpoint.
 static bool	inCone(int i, int j, int n, const int* verts, int* indices)
 {
+	PRAGMA_DISABLE_SHADOW_VARIABLE_WARNINGS
 	const int* pi = &verts[(indices[i] & 0x0fffffff) * 4];
+	PRAGMA_POP
 	const int* pj = &verts[(indices[j] & 0x0fffffff) * 4];
 	const int* pi1 = &verts[(indices[next(i, n)] & 0x0fffffff) * 4];
 	const int* pin1 = &verts[(indices[prev(i, n)] & 0x0fffffff) * 4];
