Fixed remaining navigation-related variable shadowing #UE4

diff --git a/Engine/Source/Runtime/Navmesh/Private/DetourTileCache/DetourTileCacheBuilder.cpp b/Engine/Source/Runtime/Navmesh/Private/DetourTileCache/DetourTileCacheBuilder.cpp
--- a/Engine/Source/Runtime/Navmesh/Private/DetourTileCache/DetourTileCacheBuilder.cpp
+++ b/Engine/Source/Runtime/Navmesh/Private/DetourTileCache/DetourTileCacheBuilder.cpp
@@ -1330,19 +1330,17 @@ static bool diagonalie(int i, int j, int n, const unsigned short* verts, const u
 // polygon P in the neighborhood of the i endpoint.
 static bool	inCone(int i, int j, int n, const unsigned short* verts, const unsigned short* indices)
 {
-	PRAGMA_DISABLE_SHADOW_VARIABLE_WARNINGS
-	const unsigned short* pi = &verts[(indices[i] & 0x7fff) * 4];
-	PRAGMA_POP
-	const unsigned short* pj = &verts[(indices[j] & 0x7fff) * 4];
-	const unsigned short* pi1 = &verts[(indices[TileCacheFunc::next(i, n)] & 0x7fff) * 4];
-	const unsigned short* pin1 = &verts[(indices[TileCacheFunc::prev(i, n)] & 0x7fff) * 4];
+	const unsigned short* vi = &verts[(indices[i] & 0x7fff) * 4];
+	const unsigned short* vj = &verts[(indices[j] & 0x7fff) * 4];
+	const unsigned short* vi1 = &verts[(indices[TileCacheFunc::next(i, n)] & 0x7fff) * 4];
+	const unsigned short* vin1 = &verts[(indices[TileCacheFunc::prev(i, n)] & 0x7fff) * 4];
 	
 	// If P[i] is a convex vertex [ i+1 left or on (i-1,i) ].
-	if (TileCacheFunc::leftOn(pin1, pi, pi1))
-		return TileCacheFunc::left(pi, pj, pin1) && TileCacheFunc::left(pj, pi, pi1);
+	if (TileCacheFunc::leftOn(vin1, vi, vi1))
+		return TileCacheFunc::left(vi, vj, vin1) && TileCacheFunc::left(vj, vi, vi1);
 	// Assume (i-1,i,i+1) not collinear.
 	// else P[i] is reflex.
-	return !(TileCacheFunc::leftOn(pi, pj, pi1) && TileCacheFunc::leftOn(pj, pi, pin1));
+	return !(TileCacheFunc::leftOn(vi, vj, vi1) && TileCacheFunc::leftOn(vj, vi, vin1));
 }
 
 // Returns T iff (v_i, v_j) is a proper internal
diff --git a/Engine/Source/Runtime/Navmesh/Private/Recast/RecastMesh.cpp b/Engine/Source/Runtime/Navmesh/Private/Recast/RecastMesh.cpp
--- a/Engine/Source/Runtime/Navmesh/Private/Recast/RecastMesh.cpp
+++ b/Engine/Source/Runtime/Navmesh/Private/Recast/RecastMesh.cpp
@@ -272,19 +272,17 @@ static bool diagonalie(int i, int j, int n, const int* verts, int* indices)
 // polygon P in the neighborhood of the i endpoint.
 static bool	inCone(int i, int j, int n, const int* verts, int* indices)
 {
-	PRAGMA_DISABLE_SHADOW_VARIABLE_WARNINGS
-	const int* pi = &verts[(indices[i] & 0x0fffffff) * 4];
-	PRAGMA_POP
-	const int* pj = &verts[(indices[j] & 0x0fffffff) * 4];
-	const int* pi1 = &verts[(indices[next(i, n)] & 0x0fffffff) * 4];
-	const int* pin1 = &verts[(indices[prev(i, n)] & 0x0fffffff) * 4];
+	const int* vi = &verts[(indices[i] & 0x0fffffff) * 4];
+	const int* vj = &verts[(indices[j] & 0x0fffffff) * 4];
+	const int* vi1 = &verts[(indices[next(i, n)] & 0x0fffffff) * 4];
+	const int* vin1 = &verts[(indices[prev(i, n)] & 0x0fffffff) * 4];
 
 	// If P[i] is a convex vertex [ i+1 left or on (i-1,i) ].
-	if (leftOn(pin1, pi, pi1))
-		return left(pi, pj, pin1) && left(pj, pi, pi1);
+	if (leftOn(vin1, vi, vi1))
+		return left(vi, vj, vin1) && left(vj, vi, vi1);
 	// Assume (i-1,i,i+1) not collinear.
 	// else P[i] is reflex.
-	return !(leftOn(pi, pj, pi1) && leftOn(pj, pi, pin1));
+	return !(leftOn(vi, vj, vi1) && leftOn(vj, vi, vin1));
 }
 
 // Returns T iff (v_i, v_j) is a proper internal
