fixed navmesh generation bug when single layer had huge height difference (exceeding 255 cells)
#ue4 UE-14947

diff --git a/Engine/Source/Runtime/Engine/Classes/AI/Navigation/RecastNavMesh.h b/Engine/Source/Runtime/Engine/Classes/AI/Navigation/RecastNavMesh.h
--- a/Engine/Source/Runtime/Engine/Classes/AI/Navigation/RecastNavMesh.h
+++ b/Engine/Source/Runtime/Engine/Classes/AI/Navigation/RecastNavMesh.h
@@ -17,9 +17,10 @@
 #define NAVMESHVER_64BIT				9
 #define NAVMESHVER_CLUSTER_SIMPLIFIED	10
 #define NAVMESHVER_OFFMESH_HEIGHT_BUG	11
+#define NAVMESHVER_LANDSCAPE_HEIGHT		13
 
-#define NAVMESHVER_LATEST				NAVMESHVER_OFFMESH_HEIGHT_BUG
-#define NAVMESHVER_MIN_COMPATIBLE		NAVMESHVER_CLUSTER_SIMPLIFIED
+#define NAVMESHVER_LATEST				NAVMESHVER_LANDSCAPE_HEIGHT
+#define NAVMESHVER_MIN_COMPATIBLE		NAVMESHVER_LANDSCAPE_HEIGHT
 
 #define RECAST_MAX_SEARCH_NODES		2048
 
diff --git a/Engine/Source/Runtime/Navmesh/Private/DebugUtils/DetourDebugDraw.cpp b/Engine/Source/Runtime/Navmesh/Private/DebugUtils/DetourDebugDraw.cpp
--- a/Engine/Source/Runtime/Navmesh/Private/DebugUtils/DetourDebugDraw.cpp
+++ b/Engine/Source/Runtime/Navmesh/Private/DebugUtils/DetourDebugDraw.cpp
@@ -612,7 +612,7 @@ static void debugDrawTileCachePortals(struct duDebugDraw* dd, const dtTileCacheL
 		{
 			const int idx = x+y*w;
 			const int lh = (int)layer.heights[idx];
-			if (lh == 0xff) continue;
+			if (lh == 0xffff) continue;
 			
 			for (int dir = 0; dir < 4; ++dir)
 			{
@@ -662,7 +662,7 @@ void duDebugDrawTileCacheLayerAreas(struct duDebugDraw* dd, const dtTileCacheLay
 		{
 			const int lidx = x+y*w;
 			const int lh = (int)layer.heights[lidx];
-			if (lh == 0xff) continue;
+			if (lh == 0xffff) continue;
 			const unsigned char area = layer.areas[lidx];
 			
 			unsigned int col;
@@ -720,7 +720,7 @@ void duDebugDrawTileCacheLayerDistance(struct duDebugDraw* dd, const dtTileCache
 		{
 			const int lidx = x+y*w;
 			const int lh = (int)layer.heights[lidx];
-			if (lh == 0xff) continue;
+			if (lh == 0xffff) continue;
 
 			const unsigned char cd = (unsigned char)(dfield.data[lidx] * dscale);
 			const unsigned int col = duRGBA(cd,cd,cd,255);
@@ -768,7 +768,7 @@ void duDebugDrawTileCacheLayerRegions(struct duDebugDraw* dd, const dtTileCacheL
 		{
 			const int lidx = x+y*w;
 			const int lh = (int)layer.heights[lidx];
-			if (lh == 0xff) continue;
+			if (lh == 0xffff) continue;
 			const unsigned short reg = layer.regs[lidx];
 			
 			unsigned int col = duLerpCol(color, duIntToCol(reg, 255), 192);
diff --git a/Engine/Source/Runtime/Navmesh/Private/DebugUtils/RecastDebugDraw.cpp b/Engine/Source/Runtime/Navmesh/Private/DebugUtils/RecastDebugDraw.cpp
--- a/Engine/Source/Runtime/Navmesh/Private/DebugUtils/RecastDebugDraw.cpp
+++ b/Engine/Source/Runtime/Navmesh/Private/DebugUtils/RecastDebugDraw.cpp
@@ -349,7 +349,7 @@ static void drawLayerPortals(duDebugDraw* dd, const rcHeightfieldLayer* layer)
 		{
 			const int idx = x+y*w;
 			const int lh = (int)layer->heights[idx];
-			if (lh == 255) continue;
+			if (lh == 0xffff) continue;
 			
 			for (int dir = 0; dir < 4; ++dir)
 			{
@@ -398,7 +398,7 @@ void duDebugDrawHeightfieldLayer(duDebugDraw* dd, const struct rcHeightfieldLaye
 		{
 			const int lidx = x+y*w;
 			const int lh = (int)layer.heights[lidx];
-			if (h == 0xff) continue;
+			if (h == 0xffff) continue;
 			const unsigned char area = layer.areas[lidx];
 			
 			unsigned int col;
diff --git a/Engine/Source/Runtime/Navmesh/Private/DetourTileCache/DetourTileCacheBuilder.cpp b/Engine/Source/Runtime/Navmesh/Private/DetourTileCache/DetourTileCacheBuilder.cpp
--- a/Engine/Source/Runtime/Navmesh/Private/DetourTileCache/DetourTileCacheBuilder.cpp
+++ b/Engine/Source/Runtime/Navmesh/Private/DetourTileCache/DetourTileCacheBuilder.cpp
@@ -486,7 +486,7 @@ static void simplifyContour(unsigned char area, dtTempContour& cont, const float
 	}
 }
 
-static unsigned char getCornerHeight(dtTileCacheLayer& layer, const int x, const int y, const int z, const int walkableClimb, bool& shouldRemove)
+static unsigned short getCornerHeight(dtTileCacheLayer& layer, const int x, const int y, const int z, const int walkableClimb, bool& shouldRemove)
 {
 	const int w = (int)layer.header->width;
 	const int h = (int)layer.header->height;
@@ -494,7 +494,7 @@ static unsigned char getCornerHeight(dtTileCacheLayer& layer, const int x, const
 	int n = 0;
 
 	unsigned char portal = 0xf;
-	unsigned char height = 0;
+	unsigned short height = 0;
 	unsigned short preg = 0xffff;
 	bool allSameReg = true;
 
@@ -510,7 +510,7 @@ static unsigned char getCornerHeight(dtTileCacheLayer& layer, const int x, const
 				const int lh = (int)layer.heights[idx];
 				if (dtAbs(lh-y) <= walkableClimb && layer.areas[idx] != DT_TILECACHE_NULL_AREA)
 				{
-					height = dtMax(height, (unsigned char)lh);
+					height = dtMax(height, (unsigned short)lh);
 					portal &= (layer.cons[idx] >> 4);
 					if (preg != 0xffff && preg != layer.regs[idx])
 						allSameReg = false;
@@ -811,7 +811,7 @@ dtStatus dtBuildTileCacheContours(dtTileCacheAlloc* alloc, dtTileCacheLayer& lay
 					unsigned short* vn = &temp.verts[i*5];
 					unsigned short nei = vn[3]; // The neighbour reg is stored at segment vertex of a segment. 
 					bool shouldRemove = false;
-					unsigned char lh = getCornerHeight(layer, (int)v[0], (int)v[1], (int)v[2],
+					unsigned short lh = getCornerHeight(layer, (int)v[0], (int)v[1], (int)v[2],
 						walkableClimb, shouldRemove);
 
 					if ((nei != 0xffff) && ((nei & 0xf800) == 0))
@@ -2517,7 +2517,7 @@ dtStatus dtBuildTileCacheClusters(dtTileCacheAlloc* alloc, dtTileCacheClusterSet
 
 dtStatus dtBuildTileCacheLayer(dtTileCacheCompressor* comp,
 							   dtTileCacheLayerHeader* header,
-							   const unsigned char* heights,
+							   const unsigned short* heights,
 							   const unsigned char* areas,
 							   const unsigned char* cons,
 							   unsigned char** outData, int* outDataSize)
@@ -2534,16 +2534,16 @@ dtStatus dtBuildTileCacheLayer(dtTileCacheCompressor* comp,
 	memcpy(data, header, sizeof(dtTileCacheLayerHeader));
 	
 	// Concatenate grid data for compression.
-	const int bufferSize = gridSize*3;
+	const int bufferSize = gridSize*4;
 	unsigned char* buffer = (unsigned char*)dtAlloc(bufferSize, DT_ALLOC_TEMP);
 	if (!buffer)
 	{
 		dtFree(data);
 		return DT_FAILURE | DT_OUT_OF_MEMORY;
 	}
-	memcpy(buffer, heights, gridSize);
-	memcpy(buffer+gridSize, areas, gridSize);
-	memcpy(buffer+gridSize*2, cons, gridSize);
+	memcpy(buffer, heights, gridSize*2);
+	memcpy(buffer+gridSize*2, areas, gridSize);
+	memcpy(buffer+gridSize*3, cons, gridSize);
 	
 	// Compress
 	unsigned char* compressed = data + headerSize;
@@ -2588,7 +2588,7 @@ dtStatus dtDecompressTileCacheLayer(dtTileCacheAlloc* alloc, dtTileCacheCompress
 	const int layerSize = dtAlign4(sizeof(dtTileCacheLayer));
 	const int headerSize = dtAlign4(sizeof(dtTileCacheLayerHeader));
 	const int gridSize = (int)compressedHeader->width * (int)compressedHeader->height;
-	const int bufferSize = layerSize + headerSize + gridSize*5;
+	const int bufferSize = layerSize + headerSize + gridSize*6;
 	
 	unsigned char* buffer = (unsigned char*)alloc->alloc(bufferSize);
 	if (!buffer)
@@ -2613,10 +2613,10 @@ dtStatus dtDecompressTileCacheLayer(dtTileCacheAlloc* alloc, dtTileCacheCompress
 	}
 	
 	layer->header = header;
-	layer->heights = grids;
-	layer->areas = grids + gridSize;
-	layer->cons = grids + gridSize*2;
-	layer->regs = (unsigned short*)(grids + gridSize*3);
+	layer->heights = (unsigned short*)grids;
+	layer->areas = grids + gridSize*2;
+	layer->cons = grids + gridSize*3;
+	layer->regs = (unsigned short*)(grids + gridSize*4);
 
 	*layerOut = layer;
 	
diff --git a/Engine/Source/Runtime/Navmesh/Private/DetourTileCache/DetourTileCacheDetail.cpp b/Engine/Source/Runtime/Navmesh/Private/DetourTileCache/DetourTileCacheDetail.cpp
--- a/Engine/Source/Runtime/Navmesh/Private/DetourTileCache/DetourTileCacheDetail.cpp
+++ b/Engine/Source/Runtime/Navmesh/Private/DetourTileCache/DetourTileCacheDetail.cpp
@@ -30,7 +30,7 @@
 #include <math.h>
 
 static const unsigned DT_UNSET_PATCH_HEIGHT = 0xffff;
-static const unsigned DT_UNSET_LAYER_HEIGHT = 0xff;
+static const unsigned DT_UNSET_LAYER_HEIGHT = 0xffff;
 
 struct dtHeightPatch
 {
diff --git a/Engine/Source/Runtime/Navmesh/Private/Recast/RecastLayers.cpp b/Engine/Source/Runtime/Navmesh/Private/Recast/RecastLayers.cpp
--- a/Engine/Source/Runtime/Navmesh/Private/Recast/RecastLayers.cpp
+++ b/Engine/Source/Runtime/Navmesh/Private/Recast/RecastLayers.cpp
@@ -673,14 +673,15 @@ static bool SplitAndStoreLayerRegions(rcContext* ctx, rcCompactHeightfield& chf,
 		memset(layer, 0, sizeof(rcHeightfieldLayer));
 
 		const int gridSize = sizeof(unsigned char)*lw*lh;
+		const int gridSize2 = sizeof(unsigned short)*lw*lh;
 
-		layer->heights = (unsigned char*)rcAlloc(gridSize, RC_ALLOC_PERM);
+		layer->heights = (unsigned short*)rcAlloc(gridSize2, RC_ALLOC_PERM);
 		if (!layer->heights)
 		{
-			ctx->log(RC_LOG_ERROR, "SplitAndStoreLayerRegions: Out of memory 'heights' (%d).", gridSize);
+			ctx->log(RC_LOG_ERROR, "SplitAndStoreLayerRegions: Out of memory 'heights' (%d).", gridSize2);
 			return false;
 		}
-		memset(layer->heights, 0xff, gridSize);
+		memset(layer->heights, 0xff, gridSize2);
 
 		layer->areas = (unsigned char*)rcAlloc(gridSize, RC_ALLOC_PERM);
 		if (!layer->areas)
@@ -755,7 +756,7 @@ static bool SplitAndStoreLayerRegions(rcContext* ctx, rcCompactHeightfield& chf,
 
 					// Store height and area type.
 					const int idx = x+y*lw;
-					layer->heights[idx] = (unsigned char)(s.y - hmin);
+					layer->heights[idx] = (unsigned short)(s.y - hmin);
 					layer->areas[idx] = chf.areas[j];
 
 					// Check connection.
@@ -776,7 +777,7 @@ static bool SplitAndStoreLayerRegions(rcContext* ctx, rcCompactHeightfield& chf,
 								// Update height so that it matches on both sides of the portal.
 								const rcCompactSpan& as = chf.spans[ai];
 								if (as.y > hmin)
-									layer->heights[idx] = rcMax(layer->heights[idx], (unsigned char)(as.y - hmin));
+									layer->heights[idx] = rcMax(layer->heights[idx], (unsigned short)(as.y - hmin));
 							}
 							// Valid connection mask
 							if (chf.areas[ai] != RC_NULL_AREA && lid == alid)
@@ -1320,14 +1321,15 @@ bool rcBuildHeightfieldLayers(rcContext* ctx, rcCompactHeightfield& chf,
 		memset(layer, 0, sizeof(rcHeightfieldLayer));
 
 		const int gridSize = sizeof(unsigned char)*lw*lh;
+		const int gridSize2 = sizeof(unsigned short)*lw*lh;
 
-		layer->heights = (unsigned char*)rcAlloc(gridSize, RC_ALLOC_PERM);
+		layer->heights = (unsigned short*)rcAlloc(gridSize2, RC_ALLOC_PERM);
 		if (!layer->heights)
 		{
-			ctx->log(RC_LOG_ERROR, "rcBuildHeightfieldLayers: Out of memory 'heights' (%d).", gridSize);
+			ctx->log(RC_LOG_ERROR, "rcBuildHeightfieldLayers: Out of memory 'heights' (%d).", gridSize2);
 			return false;
 		}
-		memset(layer->heights, 0xff, gridSize);
+		memset(layer->heights, 0xff, gridSize2);
 
 		layer->areas = (unsigned char*)rcAlloc(gridSize, RC_ALLOC_PERM);
 		if (!layer->areas)
@@ -1402,7 +1404,7 @@ bool rcBuildHeightfieldLayers(rcContext* ctx, rcCompactHeightfield& chf,
 
 					// Store height and area type.
 					const int idx = x+y*lw;
-					layer->heights[idx] = (unsigned char)(s.y - hmin);
+					layer->heights[idx] = (unsigned short)(s.y - hmin);
 					layer->areas[idx] = chf.areas[j];
 
 					// Check connection.
@@ -1423,7 +1425,7 @@ bool rcBuildHeightfieldLayers(rcContext* ctx, rcCompactHeightfield& chf,
 								// Update height so that it matches on both sides of the portal.
 								const rcCompactSpan& as = chf.spans[ai];
 								if (as.y > hmin)
-									layer->heights[idx] = rcMax(layer->heights[idx], (unsigned char)(as.y - hmin));
+									layer->heights[idx] = rcMax(layer->heights[idx], (unsigned short)(as.y - hmin));
 							}
 							// Valid connection mask
 							if (chf.areas[ai] != RC_NULL_AREA && lid == alid)
diff --git a/Engine/Source/Runtime/Navmesh/Public/DetourTileCache/DetourTileCacheBuilder.h b/Engine/Source/Runtime/Navmesh/Public/DetourTileCache/DetourTileCacheBuilder.h
--- a/Engine/Source/Runtime/Navmesh/Public/DetourTileCache/DetourTileCacheBuilder.h
+++ b/Engine/Source/Runtime/Navmesh/Public/DetourTileCache/DetourTileCacheBuilder.h
@@ -47,7 +47,7 @@ struct dtTileCacheLayer
 {
 	dtTileCacheLayerHeader* header;
 	unsigned short regCount;					///< Region count.
-	unsigned char* heights;
+	unsigned short* heights;
 	unsigned char* areas;
 	unsigned char* cons;
 	unsigned short* regs;
@@ -182,7 +182,7 @@ inline int getDirOffsetY(int dir)
 
 NAVMESH_API dtStatus dtBuildTileCacheLayer(dtTileCacheCompressor* comp,
 							   dtTileCacheLayerHeader* header,
-							   const unsigned char* heights,
+							   const unsigned short* heights,
 							   const unsigned char* areas,
 							   const unsigned char* cons,
 							   unsigned char** outData, int* outDataSize);
@@ -275,7 +275,8 @@ NAVMESH_API dtStatus dtBuildTileCacheClusters(dtTileCacheAlloc* alloc,
 //@UE4 END
 
 /// Swaps the endianess of the compressed tile data's header (#dtTileCacheLayerHeader).
-/// Tile layer data does not need endian swapping as it consits only of bytes.
+/// Tile layer data does not need endian swapping as it consist only of bytes.
+/// UE4: not anymore, there are short types as well now
 ///  @param[in,out]	data		The tile data array.
 ///  @param[in]		dataSize	The size of the data array.
 NAVMESH_API bool dtTileCacheHeaderSwapEndian(unsigned char* data, const int dataSize);
diff --git a/Engine/Source/Runtime/Navmesh/Public/Recast/Recast.h b/Engine/Source/Runtime/Navmesh/Public/Recast/Recast.h
--- a/Engine/Source/Runtime/Navmesh/Public/Recast/Recast.h
+++ b/Engine/Source/Runtime/Navmesh/Public/Recast/Recast.h
@@ -397,7 +397,7 @@ struct rcHeightfieldLayer
 	int maxy;					///< The maximum y-bounds of usable data. (Along the z-axis.)
 	int hmin;					///< The minimum height bounds of usable data. (Along the y-axis.)
 	int hmax;					///< The maximum height bounds of usable data. (Along the y-axis.)
-	unsigned char* heights;		///< The heightfield. [Size: (width - borderSize*2) * (h - borderSize*2)]
+	unsigned short* heights;	///< The heightfield. [Size: (width - borderSize*2) * (h - borderSize*2)]
 	unsigned char* areas;		///< Area ids. [Size: Same as #heights]
 	unsigned char* cons;		///< Packed neighbor connection information. [Size: Same as #heights]
 };
