fixed skipping custom navigation links when new path starts really close to link's entry point
#ue4

diff --git a/Engine/Source/Runtime/Navmesh/Private/Detour/DetourNavMeshQuery.cpp b/Engine/Source/Runtime/Navmesh/Private/Detour/DetourNavMeshQuery.cpp
--- a/Engine/Source/Runtime/Navmesh/Private/Detour/DetourNavMeshQuery.cpp
+++ b/Engine/Source/Runtime/Navmesh/Private/Detour/DetourNavMeshQuery.cpp
@@ -2243,7 +2243,7 @@ dtStatus dtNavMeshQuery::findStraightPath(const float* startPos, const float* en
 				}
 
 				// If starting really close the portal, advance.
-				if (i == 0)
+				if (i == 0 && toType == DT_POLYTYPE_GROUND)
 				{
 					float t;
 					if (dtDistancePtSegSqr2D(portalApex, left, right, t) < dtSqr(0.001f))
