fixed detour crowd boundary scoring
#ue4

diff --git a/Engine/Source/Runtime/Navmesh/Private/DetourCrowd/DetourLocalBoundary.cpp b/Engine/Source/Runtime/Navmesh/Private/DetourCrowd/DetourLocalBoundary.cpp
--- a/Engine/Source/Runtime/Navmesh/Private/DetourCrowd/DetourLocalBoundary.cpp
+++ b/Engine/Source/Runtime/Navmesh/Private/DetourCrowd/DetourLocalBoundary.cpp
@@ -142,7 +142,7 @@ void dtLocalBoundary::update(dtPolyRef ref, const float* pos, const float collis
 			}
 
 			// [UE4] include direction to segment in score
-			dtVmad(closestPt, s, s + 3, tseg);
+			dtVlerp(closestPt, s, s + 3, tseg);
 			dtVsub(dirToSeg, closestPt, pos);
 			dtVnormalize(dirToSeg);
 			const float dseg = dtVdot2D(dirToSeg, moveDir);
@@ -215,7 +215,7 @@ void dtLocalBoundary::update(const dtSharedBoundary* sharedData, const int share
 		}
 
 		// include direction to segment in score
-		dtVmad(closestPt, Data.Edges[Idx].v0, Data.Edges[Idx].v1, tseg);
+		dtVlerp(closestPt, Data.Edges[Idx].v0, Data.Edges[Idx].v1, tseg);
 		dtVsub(dirToSeg, closestPt, pos);
 		dtVnormalize(dirToSeg);
 		const float dseg = dtVdot2D(dirToSeg, moveDir);
