Fixed operation validity condition in dtObstacleAvoidanceQuery::addSegment #UE4

github PR#1554 by slonopotamus

diff --git a/Engine/Source/Runtime/Navmesh/Private/DetourCrowd/DetourObstacleAvoidance.cpp b/Engine/Source/Runtime/Navmesh/Private/DetourCrowd/DetourObstacleAvoidance.cpp
--- a/Engine/Source/Runtime/Navmesh/Private/DetourCrowd/DetourObstacleAvoidance.cpp
+++ b/Engine/Source/Runtime/Navmesh/Private/DetourCrowd/DetourObstacleAvoidance.cpp
@@ -275,7 +275,8 @@ void dtObstacleAvoidanceQuery::addCircle(const float* pos, const float rad,
 
 void dtObstacleAvoidanceQuery::addSegment(const float* p, const float* q)
 {
-	if (m_nsegments > m_maxSegments)
+	// [UE4] fixed condition below. Used to be strict > comparison
+	if (m_nsegments >= m_maxSegments)
 		return;
 	
 	dtObstacleSegment* seg = &m_segments[m_nsegments++];
