Fixed Detour's CollectRegionsMonotone and CollectRegionsChunky to not fail as 'out of memory' when no regions were found #UE4

Instead, if the region count is 0 we simply return as 'success' since no issues occured. The rest of nav generaiton process is handling this result properly.

diff --git a/Engine/Source/Runtime/Navmesh/Private/DetourTileCache/DetourTileCacheBuilder.cpp b/Engine/Source/Runtime/Navmesh/Private/DetourTileCache/DetourTileCacheBuilder.cpp
--- a/Engine/Source/Runtime/Navmesh/Private/DetourTileCache/DetourTileCacheBuilder.cpp
+++ b/Engine/Source/Runtime/Navmesh/Private/DetourTileCache/DetourTileCacheBuilder.cpp
@@ -1939,6 +1939,17 @@ dtStatus dtBuildTileCachePolyMesh(dtTileCacheAlloc* alloc,
 	
 	mesh.nvp = MAX_VERTS_PER_POLY;
 	
+// @UE4 BEGIN: special handling of "no valid contours"
+	if (maxVertices == 0)
+	{
+		// treating this as success because no issues arised
+		// there's just nothing to do.
+		// Note that 'mesh' properties are properly initialized to 0 at this point
+		// by dtAllocTileCachePolyMesh
+		return DT_SUCCESS;
+	}
+// @UE4 END
+
 	dtFixedArray<unsigned char> vflags(alloc, maxVertices);
 	if (!vflags)
 		return DT_FAILURE | DT_OUT_OF_MEMORY;
@@ -2510,6 +2521,15 @@ dtStatus dtReplaceArea(dtTileCacheLayer& layer, const unsigned char areaId, cons
 dtStatus dtBuildTileCacheClusters(dtTileCacheAlloc* alloc, dtTileCacheClusterSet& lclusters, dtTileCachePolyMesh& lmesh)
 {
 	lclusters.npolys = lmesh.npolys;
+// @UE4 BEGIN: special handling of "no polys"
+	if (lmesh.npolys == 0)
+	{
+		// treating this as success there's just nothing to do
+		// Note that at this point lclusters is properly initialized to 0
+		// by dtAllocTileCacheClusterSet
+		return DT_SUCCESS;
+	}
+// @UE4 END
 	lclusters.polyMap = (unsigned short*)alloc->alloc(sizeof(unsigned short)*lclusters.npolys);
 	if (!lclusters.polyMap)
 		return DT_FAILURE | DT_OUT_OF_MEMORY;
diff --git a/Engine/Source/Runtime/Navmesh/Private/DetourTileCache/DetourTileCacheRegion.cpp b/Engine/Source/Runtime/Navmesh/Private/DetourTileCache/DetourTileCacheRegion.cpp
--- a/Engine/Source/Runtime/Navmesh/Private/DetourTileCache/DetourTileCacheRegion.cpp
+++ b/Engine/Source/Runtime/Navmesh/Private/DetourTileCache/DetourTileCacheRegion.cpp
@@ -1151,6 +1151,17 @@ static dtStatus CollectRegionsMonotone(dtTileCacheAlloc* alloc, dtTileCacheLayer
 
 	// Allocate and init layer regions.
 	nregs = (int)regId;
+
+// @UE4 BEGIN: special handling of "no regions"
+	if (nregs == 0)
+	{
+		regs = 0;
+		// treating this as success because we successfully generated 0 regions,
+		// no issues occurred, everything was good. Just no regions. 
+		return DT_SUCCESS;
+	}
+// @UE4 END
+
 	regs = (dtLayerMonotoneRegion*)alloc->alloc(sizeof(dtLayerMonotoneRegion) * nregs);
 	if (!regs)
 		return DT_FAILURE | DT_OUT_OF_MEMORY;
@@ -1311,6 +1322,17 @@ static dtStatus CollectRegionsChunky(dtTileCacheAlloc* alloc, dtTileCacheLayer&
 
 	// Allocate and init layer regions.
 	nregs = (int)regId;
+
+// @UE4 BEGIN: special handling of "no regions"
+	if (nregs == 0)
+	{
+		regs = 0;
+		// treating this as success because we successfully generated 0 regions,
+		// no issues occurred, everything was good. Just no regions. 
+		return DT_SUCCESS;
+	}
+// @UE4 END
+
 	regs = (dtLayerMonotoneRegion*)alloc->alloc(sizeof(dtLayerMonotoneRegion) * nregs);
 	if (!regs)
 		return DT_FAILURE | DT_OUT_OF_MEMORY;
@@ -1529,7 +1551,9 @@ dtStatus dtBuildTileCacheRegionsMonotone(dtTileCacheAlloc* alloc, const int minR
 	int nregs = 0;
 
 	dtStatus status = CollectRegionsMonotone(alloc, layer, regs, nregs);
-	if (dtStatusSucceed(status))
+	// having no regions and status being successful is a valid state
+	// we can avoid calling MergeAndCompressRegions if that happens
+	if (dtStatusSucceed(status) && nregs > 0)
 	{
 		MergeAndCompressRegions(alloc, layer, regs, nregs, minRegionArea, mergeRegionArea);
 	}
@@ -1544,7 +1568,9 @@ dtStatus dtBuildTileCacheRegionsChunky(dtTileCacheAlloc* alloc, const int minReg
 	int nregs = 0;
 
 	dtStatus status = CollectRegionsChunky(alloc, layer, regionChunkSize, regs, nregs);
-	if (dtStatusSucceed(status))
+	// having no regions and status being successful is a valid state
+	// we can avoid calling MergeAndCompressRegions if that happens
+	if (dtStatusSucceed(status) && nregs > 0)
 	{
 		MergeAndCompressRegions(alloc, layer, regs, nregs, minRegionArea, mergeRegionArea);
 	}
