Fixed Detour's CollectLayerRegionsMonotone and CollectLayerRegionsChunky to not fail as 'out of memory' when no regions were found #UE4

Instead, if the region count is 0 we simply return as 'success' since no issues occured. The rest of nav generaiton process is handling this result properly.
This fix is the exact same fix done to DetourTileCacheRegion.cpp in 5624049.

diff --git a/Engine/Source/Runtime/Navmesh/Private/Recast/RecastLayers.cpp b/Engine/Source/Runtime/Navmesh/Private/Recast/RecastLayers.cpp
--- a/Engine/Source/Runtime/Navmesh/Private/Recast/RecastLayers.cpp
+++ b/Engine/Source/Runtime/Navmesh/Private/Recast/RecastLayers.cpp
@@ -226,6 +226,17 @@ static bool CollectLayerRegionsMonotone(rcContext* ctx, rcCompactHeightfield& ch
 
 	// Allocate and init layer regions.
 	nregs = (int)regId;
+
+// @UE4 BEGIN: special handling of "no regions"
+	if (nregs == 0)
+	{
+		regs = 0;
+		// treating this as success because we successfully generated 0 regions,
+		// no issues occurred, everything was good. Just no regions. 
+		return true;
+	}
+// @UE4 END
+
 	regs = (rcLayerRegionMonotone*)rcAlloc(sizeof(rcLayerRegionMonotone)*nregs, RC_ALLOC_TEMP);
 	if (!regs)
 	{
@@ -431,6 +442,17 @@ static bool CollectLayerRegionsChunky(rcContext* ctx, rcCompactHeightfield& chf,
 
 	// Allocate and init layer regions.
 	nregs = (int)regId;
+
+// @UE4 BEGIN: special handling of "no regions"
+	if (nregs == 0)
+	{
+		regs = 0;
+		// treating this as success because we successfully generated 0 regions,
+		// no issues occurred, everything was good. Just no regions. 
+		return true;
+	}
+// @UE4 END
+
 	regs = (rcLayerRegionMonotone*)rcAlloc(sizeof(rcLayerRegionMonotone)*nregs, RC_ALLOC_TEMP);
 	if (!regs)
 	{
