Support for a max cost in the Detour pathfinders and addition of the parameter in the navigation query
DetourNavMeshQuery -> Fix issue in updateSlicedFindPath not behaving like findPath (missing check for DT_UNWALKABLE_POLY_COST)
Looking forward to merge those or to make sure that the NavigationTestingActor uses the same pathfinder to generate its path than the one use to display debug steps

diff --git a/Engine/Source/Runtime/Engine/Public/AI/Navigation/NavQueryFilter.h b/Engine/Source/Runtime/Engine/Public/AI/Navigation/NavQueryFilter.h
--- a/Engine/Source/Runtime/Engine/Public/AI/Navigation/NavQueryFilter.h
+++ b/Engine/Source/Runtime/Engine/Public/AI/Navigation/NavQueryFilter.h
@@ -18,6 +18,7 @@ public:
 	virtual void GetAllAreaCosts(float* CostArray, float* FixedCostArray, const int32 Count) const = 0;
 	virtual void SetBacktrackingEnabled(const bool bBacktracking) = 0;
 	virtual bool IsBacktrackingEnabled() const = 0;
+	virtual float GetHeuristicScale() const = 0;
 	virtual bool IsEqual(const INavigationQueryFilterInterface* Other) const = 0;
 	virtual void SetIncludeFlags(uint16 Flags) = 0;
 	virtual uint16 GetIncludeFlags() const = 0;
@@ -76,6 +77,9 @@ public:
 	/** get node limit for A* loop */
 	FORCEINLINE uint32 GetMaxSearchNodes() const { return MaxSearchNodes; }
 
+	/** get heuristic scaling factor */
+	float GetHeuristicScale() const { return QueryFilterImpl->GetHeuristicScale(); }
+
 	/** mark filter as backtracking - parse directional links in opposite direction
 	*  (find path from End to Start, but all links works like on path from Start to End) */
 	void SetBacktrackingEnabled(const bool bBacktracking) { QueryFilterImpl->SetBacktrackingEnabled(bBacktracking); }
diff --git a/Engine/Source/Runtime/NavigationSystem/Private/NavMesh/PImplRecastNavMesh.cpp b/Engine/Source/Runtime/NavigationSystem/Private/NavMesh/PImplRecastNavMesh.cpp
--- a/Engine/Source/Runtime/NavigationSystem/Private/NavMesh/PImplRecastNavMesh.cpp
+++ b/Engine/Source/Runtime/NavigationSystem/Private/NavMesh/PImplRecastNavMesh.cpp
@@ -206,6 +206,11 @@ bool FRecastQueryFilter::IsBacktrackingEnabled() const
 	return getIsBacktracking();
 }
 
+float FRecastQueryFilter::GetHeuristicScale() const
+{
+	return getHeuristicScale();
+}
+
 bool FRecastQueryFilter::IsEqual(const INavigationQueryFilterInterface* Other) const
 {
 	// @NOTE: not type safe, should be changed when another filter type is introduced
@@ -777,8 +782,13 @@ void FPImplRecastNavMesh::Raycast(const FVector& StartLoc, const FVector& EndLoc
 	}
 }
 
+ENavigationQueryResult::Type FPImplRecastNavMesh::FindPath(const FVector& StartLoc, const FVector& EndLoc, FNavMeshPath& Path, const FNavigationQueryFilter& Filter, const UObject* Owner) const
+{
+	return FindPath(StartLoc, EndLoc, FLT_MAX, Path, Filter, Owner);
+}
+
 // @TODONAV
-ENavigationQueryResult::Type FPImplRecastNavMesh::FindPath(const FVector& StartLoc, const FVector& EndLoc, FNavMeshPath& Path, const FNavigationQueryFilter& InQueryFilter, const UObject* Owner) const
+ENavigationQueryResult::Type FPImplRecastNavMesh::FindPath(const FVector& StartLoc, const FVector& EndLoc, const float MaxCost, FNavMeshPath& Path, const FNavigationQueryFilter& InQueryFilter, const UObject* Owner) const
 {
 	// temporarily disabling this check due to it causing too much "crashes"
 	// @todo but it needs to be back at some point since it realy checks for a buggy setup
@@ -816,7 +826,7 @@ ENavigationQueryResult::Type FPImplRecastNavMesh::FindPath(const FVector& StartL
 
 	// get path corridor
 	dtQueryResult PathResult;
-	const dtStatus FindPathStatus = NavQuery.findPath(StartPolyID, EndPolyID, &RecastStartPos.X, &RecastEndPos.X, QueryFilter, PathResult, 0);
+	const dtStatus FindPathStatus = NavQuery.findPath(StartPolyID, EndPolyID, &RecastStartPos.X, &RecastEndPos.X, MaxCost, QueryFilter, PathResult, 0);
 
 	// check for special case, where path has not been found, and starting polygon
 	// was the one closest to the target
@@ -887,8 +897,9 @@ ENavigationQueryResult::Type FPImplRecastNavMesh::TestPath(const FVector& StartL
 
 	// get path corridor
 	dtQueryResult PathResult;
+	const float MaxCost = FLT_MAX;
 	const dtStatus FindPathStatus = NavQuery.findPath(StartPolyID, EndPolyID,
-		&RecastStartPos.X, &RecastEndPos.X, QueryFilter, PathResult, 0);
+		&RecastStartPos.X, &RecastEndPos.X, MaxCost, QueryFilter, PathResult, 0);
 
 	if (NumVisitedNodes)
 	{
@@ -1266,7 +1277,7 @@ static void StorePathfindingDebugStep(const dtNavMeshQuery& NavQuery, const dtNa
 	}
 }
 
-int32 FPImplRecastNavMesh::DebugPathfinding(const FVector& StartLoc, const FVector& EndLoc, const FNavigationQueryFilter& Filter, const UObject* Owner, TArray<FRecastDebugPathfindingData>& Steps)
+int32 FPImplRecastNavMesh::DebugPathfinding(const FVector& StartLoc, const FVector& EndLoc, const float MaxCost, const FNavigationQueryFilter& Filter, const UObject* Owner, TArray<FRecastDebugPathfindingData>& Steps)
 {
 	int32 NumSteps = 0;
 
@@ -1288,7 +1299,7 @@ int32 FPImplRecastNavMesh::DebugPathfinding(const FVector& StartLoc, const FVect
 		return NumSteps;
 	}
 
-	dtStatus status = NavQuery.initSlicedFindPath(StartPolyID, EndPolyID, &RecastStartPos.X, &RecastEndPos.X, QueryFilter);
+	dtStatus status = NavQuery.initSlicedFindPath(StartPolyID, EndPolyID, &RecastStartPos.X, &RecastEndPos.X, MaxCost, QueryFilter);
 	while (dtStatusInProgress(status))
 	{
 		StorePathfindingDebugStep(NavQuery, DetourNavMesh, Steps);
diff --git a/Engine/Source/Runtime/NavigationSystem/Private/NavMesh/RecastNavMesh.cpp b/Engine/Source/Runtime/NavigationSystem/Private/NavMesh/RecastNavMesh.cpp
--- a/Engine/Source/Runtime/NavigationSystem/Private/NavMesh/RecastNavMesh.cpp
+++ b/Engine/Source/Runtime/NavigationSystem/Private/NavMesh/RecastNavMesh.cpp
@@ -1310,7 +1310,8 @@ ENavigationQueryResult::Type ARecastNavMesh::CalcPathLengthAndCost(const FVector
 			Path->SetWantsStringPulling(false);
 			Path->SetWantsPathCorridor(true);
 			
-			Result = RecastNavMeshImpl->FindPath(PathStart, PathEnd, Path.Get(), GetRightFilterRef(QueryFilter), QueryOwner);
+			const float MaxCost = FLT_MAX;
+			Result = RecastNavMeshImpl->FindPath(PathStart, PathEnd, MaxCost, Path.Get(), GetRightFilterRef(QueryFilter), QueryOwner);
 
 			if (Result == ENavigationQueryResult::Success || (Result == ENavigationQueryResult::Fail && Path->IsPartial()))
 			{
@@ -2094,7 +2095,7 @@ FPathFindingResult ARecastNavMesh::FindPath(const FNavAgentProperties& AgentProp
 		}
 		else
 		{
-			Result.Result = RecastNavMesh->RecastNavMeshImpl->FindPath(Query.StartLocation, AdjustedEndLocation, *NavMeshPath, *NavFilter, Query.Owner.Get());
+			Result.Result = RecastNavMesh->RecastNavMeshImpl->FindPath(Query.StartLocation, AdjustedEndLocation, Query.MaxCost, *NavMeshPath, *NavFilter, Query.Owner.Get());
 
 			const bool bPartialPath = Result.IsPartial();
 			if (bPartialPath)
@@ -2308,7 +2309,7 @@ int32 ARecastNavMesh::DebugPathfinding(const FPathFindingQuery& Query, TArray<FR
 
 	if ((Query.StartLocation - Query.EndLocation).IsNearlyZero() == false)
 	{
-		NumSteps = RecastNavMesh->RecastNavMeshImpl->DebugPathfinding(Query.StartLocation, Query.EndLocation, *(Query.QueryFilter.Get()), Query.Owner.Get(), Steps);
+		NumSteps = RecastNavMesh->RecastNavMeshImpl->DebugPathfinding(Query.StartLocation, Query.EndLocation, Query.MaxCost, *(Query.QueryFilter.Get()), Query.Owner.Get(), Steps);
 	}
 
 	return NumSteps;
diff --git a/Engine/Source/Runtime/NavigationSystem/Private/NavigationData.cpp b/Engine/Source/Runtime/NavigationSystem/Private/NavigationData.cpp
--- a/Engine/Source/Runtime/NavigationSystem/Private/NavigationData.cpp
+++ b/Engine/Source/Runtime/NavigationSystem/Private/NavigationData.cpp
@@ -19,8 +19,8 @@
 //----------------------------------------------------------------------//
 // FPathFindingQuery
 //----------------------------------------------------------------------//
-FPathFindingQuery::FPathFindingQuery(const UObject* InOwner, const ANavigationData& InNavData, const FVector& Start, const FVector& End, FSharedConstNavQueryFilter SourceQueryFilter, FNavPathSharedPtr InPathInstanceToFill) :
-	FPathFindingQueryData(InOwner, Start, End, SourceQueryFilter), 
+FPathFindingQuery::FPathFindingQuery(const UObject* InOwner, const ANavigationData& InNavData, const FVector& Start, const FVector& End, FSharedConstNavQueryFilter SourceQueryFilter, FNavPathSharedPtr InPathInstanceToFill, const float MaxCost) :
+	FPathFindingQueryData(InOwner, Start, End, SourceQueryFilter, MaxCost),
 	NavData(&InNavData), PathInstanceToFill(InPathInstanceToFill), NavAgentProperties(InNavData.GetConfig())
 {
 	if (!QueryFilter.IsValid() && NavData.IsValid())
@@ -29,8 +29,8 @@ FPathFindingQuery::FPathFindingQuery(const UObject* InOwner, const ANavigationDa
 	}
 }
 
-FPathFindingQuery::FPathFindingQuery(const INavAgentInterface& InNavAgent, const ANavigationData& InNavData, const FVector& Start, const FVector& End, FSharedConstNavQueryFilter SourceQueryFilter, FNavPathSharedPtr InPathInstanceToFill) :
-	FPathFindingQueryData(Cast<UObject>(&InNavAgent), Start, End, SourceQueryFilter),
+FPathFindingQuery::FPathFindingQuery(const INavAgentInterface& InNavAgent, const ANavigationData& InNavData, const FVector& Start, const FVector& End, FSharedConstNavQueryFilter SourceQueryFilter, FNavPathSharedPtr InPathInstanceToFill, const float MaxCost) :
+	FPathFindingQueryData(Cast<UObject>(&InNavAgent), Start, End, SourceQueryFilter, MaxCost),
 	NavData(&InNavData), PathInstanceToFill(InPathInstanceToFill), NavAgentProperties(InNavAgent.GetNavAgentPropertiesRef())
 {
 	if (!QueryFilter.IsValid() && NavData.IsValid())
@@ -40,7 +40,7 @@ FPathFindingQuery::FPathFindingQuery(const INavAgentInterface& InNavAgent, const
 }
 
 FPathFindingQuery::FPathFindingQuery(const FPathFindingQuery& Source) :
-	FPathFindingQueryData(Source.Owner.Get(), Source.StartLocation, Source.EndLocation, Source.QueryFilter, Source.NavDataFlags, Source.bAllowPartialPaths),
+	FPathFindingQueryData(Source.Owner.Get(), Source.StartLocation, Source.EndLocation, Source.QueryFilter, Source.NavDataFlags, Source.bAllowPartialPaths, Source.MaxCost),
 	NavData(Source.NavData), PathInstanceToFill(Source.PathInstanceToFill), NavAgentProperties(Source.NavAgentProperties)
 {
 	if (!QueryFilter.IsValid() && NavData.IsValid())
@@ -82,12 +82,25 @@ FPathFindingQuery::FPathFindingQuery(FNavPathSharedRef PathToRecalculate, const
 	}
 }
 
+float FPathFindingQuery::ComputeMaxCostFromHeuristic(const FVector& StartPos, const FVector& EndPos, const float HeuristicScale, const float MaxCostFactor) const
+{
+	if (MaxCostFactor == FLT_MAX)
+	{
+		return FLT_MAX;
+	}
+	else
+	{
+		const float OriginalHeuristicEstimate = HeuristicScale * FVector::Dist(StartPos, EndPos);
+		return FMath::Min(FLT_MAX, MaxCostFactor * OriginalHeuristicEstimate);
+	}
+}
+
 //----------------------------------------------------------------------//
 // FAsyncPathFindingQuery
 //----------------------------------------------------------------------//
 uint32 FAsyncPathFindingQuery::LastPathFindingUniqueID = INVALID_NAVQUERYID;
 
-FAsyncPathFindingQuery::FAsyncPathFindingQuery(const UObject* InOwner, const ANavigationData& InNavData, const FVector& Start, const FVector& End, const FNavPathQueryDelegate& Delegate, FSharedConstNavQueryFilter SourceQueryFilter)
+FAsyncPathFindingQuery::FAsyncPathFindingQuery(const UObject* InOwner, const ANavigationData& InNavData, const FVector& Start, const FVector& End, const FNavPathQueryDelegate& Delegate, FSharedConstNavQueryFilter SourceQueryFilter, const float MaxCost)
 : FPathFindingQuery(InOwner, InNavData, Start, End, SourceQueryFilter)
 , QueryID(GetUniqueID())
 , OnDoneDelegate(Delegate)
@@ -104,6 +117,7 @@ FAsyncPathFindingQuery::FAsyncPathFindingQuery(const FPathFindingQuery& Query, c
 {
 
 }
+
 //----------------------------------------------------------------------//
 // FSupportedAreaData
 //----------------------------------------------------------------------//
diff --git a/Engine/Source/Runtime/NavigationSystem/Public/AbstractNavData.h b/Engine/Source/Runtime/NavigationSystem/Public/AbstractNavData.h
--- a/Engine/Source/Runtime/NavigationSystem/Public/AbstractNavData.h
+++ b/Engine/Source/Runtime/NavigationSystem/Public/AbstractNavData.h
@@ -29,6 +29,7 @@ public:
 	virtual void GetAllAreaCosts(float* CostArray, float* FixedCostArray, const int32 Count) const override {}
 	virtual void SetBacktrackingEnabled(const bool bBacktracking) override {}
 	virtual bool IsBacktrackingEnabled() const override { return false; }
+	virtual float GetHeuristicScale() const override { return 1.f; }
 	virtual bool IsEqual(const INavigationQueryFilterInterface* Other) const override { return true; }
 	virtual void SetIncludeFlags(uint16 Flags) override {}
 	virtual uint16 GetIncludeFlags() const override { return 0; }
diff --git a/Engine/Source/Runtime/NavigationSystem/Public/NavMesh/PImplRecastNavMesh.h b/Engine/Source/Runtime/NavigationSystem/Public/NavMesh/PImplRecastNavMesh.h
--- a/Engine/Source/Runtime/NavigationSystem/Public/NavMesh/PImplRecastNavMesh.h
+++ b/Engine/Source/Runtime/NavigationSystem/Public/NavMesh/PImplRecastNavMesh.h
@@ -72,7 +72,11 @@ public:
 
 	// @TODONAV
 	/** Generates path from the given query. Synchronous. */
+	UE_DEPRECATED(4.25, "Use the version with the added MaxCost parameter (FLT_MAX can be used as default).")
 	ENavigationQueryResult::Type FindPath(const FVector& StartLoc, const FVector& EndLoc, FNavMeshPath& Path, const FNavigationQueryFilter& Filter, const UObject* Owner) const;
+	
+	/** Generates path from the given query. Synchronous. */
+	ENavigationQueryResult::Type FindPath(const FVector& StartLoc, const FVector& EndLoc, const float MaxCost, FNavMeshPath& Path, const FNavigationQueryFilter& Filter, const UObject* Owner) const;
 
 	/** Check if path exists */
 	ENavigationQueryResult::Type TestPath(const FVector& StartLoc, const FVector& EndLoc, const FNavigationQueryFilter& Filter, const UObject* Owner, int32* NumVisitedNodes = 0) const;
@@ -85,7 +89,7 @@ public:
 		ARecastNavMesh::FRaycastResult& RaycastResult, NavNodeRef StartNode = INVALID_NAVNODEREF) const;
 
 	/** Generates path from given query and collect data for every step of A* algorithm */
-	int32 DebugPathfinding(const FVector& StartLoc, const FVector& EndLoc, const FNavigationQueryFilter& Filter, const UObject* Owner, TArray<FRecastDebugPathfindingData>& Steps);
+	int32 DebugPathfinding(const FVector& StartLoc, const FVector& EndLoc, const float MaxCost, const FNavigationQueryFilter& Filter, const UObject* Owner, TArray<FRecastDebugPathfindingData>& Steps);
 
 	/** Returns a random location on the navmesh. */
 	FNavLocation GetRandomPoint(const FNavigationQueryFilter& Filter, const UObject* Owner) const;
diff --git a/Engine/Source/Runtime/NavigationSystem/Public/NavMesh/RecastQueryFilter.h b/Engine/Source/Runtime/NavigationSystem/Public/NavMesh/RecastQueryFilter.h
--- a/Engine/Source/Runtime/NavigationSystem/Public/NavMesh/RecastQueryFilter.h
+++ b/Engine/Source/Runtime/NavigationSystem/Public/NavMesh/RecastQueryFilter.h
@@ -29,6 +29,7 @@ public:
 	virtual void GetAllAreaCosts(float* CostArray, float* FixedCostArray, const int32 Count) const override;
 	virtual void SetBacktrackingEnabled(const bool bBacktracking) override;
 	virtual bool IsBacktrackingEnabled() const override;
+	virtual float GetHeuristicScale() const override;
 	virtual bool IsEqual(const INavigationQueryFilterInterface* Other) const override;
 	virtual void SetIncludeFlags(uint16 Flags) override;
 	virtual uint16 GetIncludeFlags() const override;
diff --git a/Engine/Source/Runtime/NavigationSystem/Public/NavigationData.h b/Engine/Source/Runtime/NavigationSystem/Public/NavigationData.h
--- a/Engine/Source/Runtime/NavigationSystem/Public/NavigationData.h
+++ b/Engine/Source/Runtime/NavigationSystem/Public/NavigationData.h
@@ -1002,7 +1002,7 @@ struct FAsyncPathFindingQuery : public FPathFindingQuery
 		, Mode(EPathFindingMode::Regular)
 	{ }
 
-	FAsyncPathFindingQuery(const UObject* InOwner, const ANavigationData& InNavData, const FVector& Start, const FVector& End, const FNavPathQueryDelegate& Delegate, FSharedConstNavQueryFilter SourceQueryFilter);
+	FAsyncPathFindingQuery(const UObject* InOwner, const ANavigationData& InNavData, const FVector& Start, const FVector& End, const FNavPathQueryDelegate& Delegate, FSharedConstNavQueryFilter SourceQueryFilter, const float MaxCost = FLT_MAX);
 	FAsyncPathFindingQuery(const FPathFindingQuery& Query, const FNavPathQueryDelegate& Delegate, const EPathFindingMode::Type QueryMode);
 
 protected:
diff --git a/Engine/Source/Runtime/NavigationSystem/Public/NavigationSystemTypes.h b/Engine/Source/Runtime/NavigationSystem/Public/NavigationSystemTypes.h
--- a/Engine/Source/Runtime/NavigationSystem/Public/NavigationSystemTypes.h
+++ b/Engine/Source/Runtime/NavigationSystem/Public/NavigationSystemTypes.h
@@ -27,15 +27,19 @@ struct NAVIGATIONSYSTEM_API FPathFindingQueryData
 	FVector EndLocation;
 	FSharedConstNavQueryFilter QueryFilter;
 
+	/** cost limit of nodes allowed to be added to the open list */
+	float MaxCost;
+
 	/** additional flags passed to navigation data handling request */
 	int32 NavDataFlags;
 
 	/** if set, allow partial paths as a result */
 	uint32 bAllowPartialPaths : 1;
 
-	FPathFindingQueryData() : StartLocation(FNavigationSystem::InvalidLocation), EndLocation(FNavigationSystem::InvalidLocation), NavDataFlags(0), bAllowPartialPaths(true) {}
-	FPathFindingQueryData(const UObject* InOwner, const FVector& InStartLocation, const FVector& InEndLocation, FSharedConstNavQueryFilter InQueryFilter = nullptr, int32 InNavDataFlags = 0, bool bInAllowPartialPaths = true) :
-		Owner(InOwner), StartLocation(InStartLocation), EndLocation(InEndLocation), QueryFilter(InQueryFilter), NavDataFlags(InNavDataFlags), bAllowPartialPaths(bInAllowPartialPaths) {}
+	FPathFindingQueryData() : StartLocation(FNavigationSystem::InvalidLocation), EndLocation(FNavigationSystem::InvalidLocation), MaxCost(FLT_MAX), NavDataFlags(0), bAllowPartialPaths(true) {}
+
+	FPathFindingQueryData(const UObject* InOwner, const FVector& InStartLocation, const FVector& InEndLocation, FSharedConstNavQueryFilter InQueryFilter = nullptr, int32 InNavDataFlags = 0, bool bInAllowPartialPaths = true, const float InMaxCost = FLT_MAX) :
+		Owner(InOwner), StartLocation(InStartLocation), EndLocation(InEndLocation), QueryFilter(InQueryFilter), MaxCost(InMaxCost), NavDataFlags(InNavDataFlags), bAllowPartialPaths(bInAllowPartialPaths) {}
 };
 
 struct NAVIGATIONSYSTEM_API FPathFindingQuery : public FPathFindingQueryData
@@ -46,14 +50,17 @@ struct NAVIGATIONSYSTEM_API FPathFindingQuery : public FPathFindingQueryData
 
 	FPathFindingQuery() : FPathFindingQueryData() {}
 	FPathFindingQuery(const FPathFindingQuery& Source);
-	FPathFindingQuery(const UObject* InOwner, const ANavigationData& InNavData, const FVector& Start, const FVector& End, FSharedConstNavQueryFilter SourceQueryFilter = NULL, FNavPathSharedPtr InPathInstanceToFill = NULL);
-	FPathFindingQuery(const INavAgentInterface& InNavAgent, const ANavigationData& InNavData, const FVector& Start, const FVector& End, FSharedConstNavQueryFilter SourceQueryFilter = NULL, FNavPathSharedPtr InPathInstanceToFill = NULL);
+	FPathFindingQuery(const UObject* InOwner, const ANavigationData& InNavData, const FVector& Start, const FVector& End, FSharedConstNavQueryFilter SourceQueryFilter = NULL, FNavPathSharedPtr InPathInstanceToFill = NULL, const float MaxCost = FLT_MAX);
+	FPathFindingQuery(const INavAgentInterface& InNavAgent, const ANavigationData& InNavData, const FVector& Start, const FVector& End, FSharedConstNavQueryFilter SourceQueryFilter = NULL, FNavPathSharedPtr InPathInstanceToFill = NULL, const float MaxCost = FLT_MAX);
 
 	explicit FPathFindingQuery(FNavPathSharedRef PathToRecalculate, const ANavigationData* NavDataOverride = NULL);
 
 	FPathFindingQuery& SetPathInstanceToUpdate(FNavPathSharedPtr InPathInstanceToFill) { PathInstanceToFill = InPathInstanceToFill; return *this; }
 	FPathFindingQuery& SetAllowPartialPaths(bool bAllow) { bAllowPartialPaths = bAllow; return *this; }
 	FPathFindingQuery& SetNavAgentProperties(const FNavAgentProperties& InNavAgentProperties) { NavAgentProperties = InNavAgentProperties; return *this; }
+
+	/** utility function to compute a max cost using an Euclidean heuristic, an heuristic scale and a max cost factor */
+	float ComputeMaxCostFromHeuristic(const FVector& StartPos, const FVector& EndPos, const float HeuristicScale, const float MaxCostFactor) const;
 };
 
 namespace EPathFindingMode
diff --git a/Engine/Source/Runtime/Navmesh/Private/Detour/DetourNavMeshQuery.cpp b/Engine/Source/Runtime/Navmesh/Private/Detour/DetourNavMeshQuery.cpp
--- a/Engine/Source/Runtime/Navmesh/Private/Detour/DetourNavMeshQuery.cpp
+++ b/Engine/Source/Runtime/Navmesh/Private/Detour/DetourNavMeshQuery.cpp
@@ -1423,7 +1423,7 @@ dtStatus dtNavMeshQuery::queryPolygons(const float* center, const float* extents
 ///
 dtStatus dtNavMeshQuery::findPath(dtPolyRef startRef, dtPolyRef endRef,
 								  const float* startPos, const float* endPos,
-								  const dtQueryFilter* filter,
+								  const float maxCost, const dtQueryFilter* filter, //@UE4
 								  dtQueryResult& result, float* totalCost) const
 {
 	dtAssert(m_nav);
@@ -1583,14 +1583,19 @@ dtStatus dtNavMeshQuery::findPath(dtPolyRef startRef, dtPolyRef endRef,
 			const float total = cost + heuristic;
 
 			// The node is already in open list and the new result is worse, skip.
+			if ((neighbourNode->flags & DT_NODE_OPEN) && total >= neighbourNode->total)
+				continue;
+
 			// The node is already visited and process, and the new result is worse, skip.
+			if ((neighbourNode->flags & DT_NODE_CLOSED) && total >= neighbourNode->total)
+				continue;
+
 			// Cost of current link is DT_UNWALKABLE_POLY_COST, skip.
-			if (((neighbourNode->flags & DT_NODE_OPEN) && total >= neighbourNode->total) ||
-				((neighbourNode->flags & DT_NODE_CLOSED) && total >= neighbourNode->total) ||
-				(curCost == DT_UNWALKABLE_POLY_COST))
-			{
+			if (curCost == DT_UNWALKABLE_POLY_COST)
+				continue;
+
+			if (total > maxCost) //@UE4
 				continue;
-			}
 
 			// Add or update the node.
 			neighbourNode->pidx = m_nodePool->getNodeIdx(bestNode);
@@ -1828,7 +1833,7 @@ dtStatus dtNavMeshQuery::testClusterPath(dtPolyRef startRef, dtPolyRef endRef) c
 /// path query.
 ///
 dtStatus dtNavMeshQuery::initSlicedFindPath(dtPolyRef startRef, dtPolyRef endRef,
-											const float* startPos, const float* endPos,
+											const float* startPos, const float* endPos, const float maxCost, //@UE4
 											const dtQueryFilter* filter)
 {
 	dtAssert(m_nav);
@@ -1842,6 +1847,7 @@ dtStatus dtNavMeshQuery::initSlicedFindPath(dtPolyRef startRef, dtPolyRef endRef
 	m_query.endRef = endRef;
 	dtVcopy(m_query.startPos, startPos);
 	dtVcopy(m_query.endPos, endPos);
+	m_query.maxCost = maxCost; //@UE4
 	m_query.filter = filter;
 	
 	if (!startRef || !endRef)
@@ -1998,12 +2004,13 @@ dtStatus dtNavMeshQuery::updateSlicedFindPath(const int maxIter, int* doneIters)
 			// Calculate cost and heuristic.
 			float cost = 0;
 			float heuristic = 0;
+			float curCost = 0; //@UE4
 			
 			// Special case for last node.
 			if (neighbourRef != m_query.endRef)
 			{
 				// Cost
-				const float curCost = m_query.filter->getCost(bestNode->pos, neiPos,
+				curCost = m_query.filter->getCost(bestNode->pos, neiPos,
 					parentRef, parentTile, parentPoly,
 					bestRef, bestTile, bestPoly,
 					neighbourRef, neighbourTile, neighbourPoly);
@@ -2013,27 +2020,35 @@ dtStatus dtNavMeshQuery::updateSlicedFindPath(const int maxIter, int* doneIters)
 			else
 			{
 				// Cost
-				const float curCost = m_query.filter->getCost(bestNode->pos, neiPos,
-															  parentRef, parentTile, parentPoly,
-															  bestRef, bestTile, bestPoly,
-															  neighbourRef, neighbourTile, neighbourPoly);
+				curCost = m_query.filter->getCost(bestNode->pos, neiPos,
+												  parentRef, parentTile, parentPoly,
+												  bestRef, bestTile, bestPoly,
+												  neighbourRef, neighbourTile, neighbourPoly);
 				const float endCost = m_query.filter->getCost(neiPos, m_query.endPos,
-															  bestRef, bestTile, bestPoly,
-															  neighbourRef, neighbourTile, neighbourPoly,
-															  0, 0, 0);
+												  bestRef, bestTile, bestPoly,
+												  neighbourRef, neighbourTile, neighbourPoly,
+												  0, 0, 0);
 				
 				cost = bestNode->cost + curCost + endCost;
 				heuristic = 0;
 			}
 			
 			const float total = cost + heuristic;
-			
+
 			// The node is already in open list and the new result is worse, skip.
 			if ((neighbourNode->flags & DT_NODE_OPEN) && total >= neighbourNode->total)
 				continue;
+
 			// The node is already visited and process, and the new result is worse, skip.
 			if ((neighbourNode->flags & DT_NODE_CLOSED) && total >= neighbourNode->total)
 				continue;
+
+			// Cost of current link is DT_UNWALKABLE_POLY_COST, skip.
+			if (curCost == DT_UNWALKABLE_POLY_COST) //@UE4
+				continue;
+
+			if (total > m_query.maxCost) //@UE4
+				continue;
 			
 			// Add or update the node.
 			neighbourNode->pidx = m_nodePool->getNodeIdx(bestNode);
diff --git a/Engine/Source/Runtime/Navmesh/Private/DetourCrowd/DetourCrowd.cpp b/Engine/Source/Runtime/Navmesh/Private/DetourCrowd/DetourCrowd.cpp
--- a/Engine/Source/Runtime/Navmesh/Private/DetourCrowd/DetourCrowd.cpp
+++ b/Engine/Source/Runtime/Navmesh/Private/DetourCrowd/DetourCrowd.cpp
@@ -854,7 +854,8 @@ void dtCrowd::updateMoveRequest(const float /*dt*/)
 			// Quick seach towards the goal.
 			static const int MAX_ITER = 20;
 			m_navquery->updateLinkFilter(ag->params.linkFilter.Get());
-			m_navquery->initSlicedFindPath(path[0], ag->targetRef, ag->npos, ag->targetPos, &m_filters[ag->params.filter]);
+			const float maxCost = FLT_MAX; //@UE4
+			m_navquery->initSlicedFindPath(path[0], ag->targetRef, ag->npos, ag->targetPos, maxCost, &m_filters[ag->params.filter]); //@UE4
 			m_navquery->updateSlicedFindPath(MAX_ITER, 0);
 			dtStatus status = 0;
 			if (ag->targetReplan) // && npath > 10)
@@ -920,8 +921,9 @@ void dtCrowd::updateMoveRequest(const float /*dt*/)
 	for (int i = 0; i < nqueue; ++i)
 	{
 		dtCrowdAgent* ag = queue[i];
+		const float maxCost = FLT_MAX; //@UE4
 		ag->targetPathqRef = m_pathq.request(ag->corridor.getLastPoly(), ag->targetRef,
-			ag->corridor.getTarget(), ag->targetPos, &m_filters[ag->params.filter], ag->params.linkFilter);
+			ag->corridor.getTarget(), ag->targetPos, maxCost, &m_filters[ag->params.filter], ag->params.linkFilter); //@UE4
 		if (ag->targetPathqRef != DT_PATHQ_INVALID)
 			ag->targetState = DT_CROWDAGENT_TARGET_WAITING_FOR_PATH;
 	}
diff --git a/Engine/Source/Runtime/Navmesh/Private/DetourCrowd/DetourPathCorridor.cpp b/Engine/Source/Runtime/Navmesh/Private/DetourCrowd/DetourPathCorridor.cpp
--- a/Engine/Source/Runtime/Navmesh/Private/DetourCrowd/DetourPathCorridor.cpp
+++ b/Engine/Source/Runtime/Navmesh/Private/DetourCrowd/DetourPathCorridor.cpp
@@ -524,7 +524,8 @@ bool dtPathCorridor::optimizePathTopology(dtNavMeshQuery* navquery, const dtQuer
 	
 	dtPolyRef res[MAX_RES];
 	int nres = 0;
-	navquery->initSlicedFindPath(m_path[0], m_path[m_npath-1], m_pos, m_target, filter);
+	const float maxCost = FLT_MAX; //@UE4
+	navquery->initSlicedFindPath(m_path[0], m_path[m_npath-1], m_pos, m_target, maxCost, filter); //@UE4
 	navquery->updateSlicedFindPath(MAX_ITER, 0);
 	dtStatus status = navquery->finalizeSlicedFindPathPartial(m_path, m_npath, res, &nres, MAX_RES);
 	
diff --git a/Engine/Source/Runtime/Navmesh/Private/DetourCrowd/DetourPathQueue.cpp b/Engine/Source/Runtime/Navmesh/Private/DetourCrowd/DetourPathQueue.cpp
--- a/Engine/Source/Runtime/Navmesh/Private/DetourCrowd/DetourPathQueue.cpp
+++ b/Engine/Source/Runtime/Navmesh/Private/DetourCrowd/DetourPathQueue.cpp
@@ -111,7 +111,7 @@ void dtPathQueue::update(const int maxIters)
 		// Handle query start.
 		if (q.status == 0)
 		{
-			q.status = m_navquery->initSlicedFindPath(q.startRef, q.endRef, q.startPos, q.endPos, q.filter);
+			q.status = m_navquery->initSlicedFindPath(q.startRef, q.endRef, q.startPos, q.endPos, q.maxCost, q.filter); //@UE4
 		}		
 		// Handle query in progress.
 		if (dtStatusInProgress(q.status))
@@ -133,7 +133,7 @@ void dtPathQueue::update(const int maxIters)
 }
 
 dtPathQueueRef dtPathQueue::request(dtPolyRef startRef, dtPolyRef endRef,
-									const float* startPos, const float* endPos,
+									const float* startPos, const float* endPos, const float maxCost, //@UE4
 									const dtQueryFilter* filter, TSharedPtr<dtQuerySpecialLinkFilter> linkFilter)
 {
 	// Find empty slot
@@ -159,6 +159,7 @@ dtPathQueueRef dtPathQueue::request(dtPolyRef startRef, dtPolyRef endRef,
 	q.startRef = startRef;
 	dtVcopy(q.endPos, endPos);
 	q.endRef = endRef;
+	q.maxCost = maxCost; //@UE4
 	
 	q.status = 0;
 	q.npath = 0;
diff --git a/Engine/Source/Runtime/Navmesh/Public/Detour/DetourNavMeshQuery.h b/Engine/Source/Runtime/Navmesh/Public/Detour/DetourNavMeshQuery.h
--- a/Engine/Source/Runtime/Navmesh/Public/Detour/DetourNavMeshQuery.h
+++ b/Engine/Source/Runtime/Navmesh/Public/Detour/DetourNavMeshQuery.h
@@ -65,7 +65,7 @@ struct NAVMESH_API dtQueryFilterData
 	float lowestAreaCost;
 
 	unsigned short m_includeFlags;		///< Flags for polygons that can be visited. (Used by default implementation.)
-	unsigned short m_excludeFlags;		///< Flags for polygons that should not be visted. (Used by default implementation.)
+	unsigned short m_excludeFlags;		///< Flags for polygons that should not be visited. (Used by default implementation.)
 
 	bool m_isBacktracking;
 	//@UE4 BEGIN
@@ -219,8 +219,7 @@ public:
 	///  @returns heuristic scale
 	inline float getHeuristicScale() const { return data.heuristicScale; }
 
-	/// Retrieves euclidean distance heuristic scale
-	///  @returns heuristic scale
+	/// Set euclidean distance heuristic scale
 	inline void setHeuristicScale(const float newScale) { data.heuristicScale = newScale; }
 
 	/// Filters link in regards to its side. Used for backtracking.
@@ -345,11 +344,12 @@ public:
 	///  @param[in]		endRef		The reference id of the end polygon.
 	///  @param[in]		startPos	A position within the start polygon. [(x, y, z)]
 	///  @param[in]		endPos		A position within the end polygon. [(x, y, z)]
+	///  @param[in]		maxCost		Cost limit of nodes allowed to be added to the open list	//@UE4
 	///  @param[in]		filter		The polygon filter to apply to the query.
 	///  @param[out]	result		Results for path corridor, fills in refs and costs for each poly from start to end
 	///	 @param[out]	totalCost			If provided will get filled will total cost of path
 	dtStatus findPath(dtPolyRef startRef, dtPolyRef endRef,
-					  const float* startPos, const float* endPos,
+					  const float* startPos, const float* endPos, const float maxCost, //@UE4
 					  const dtQueryFilter* filter,
 					  dtQueryResult& result, float* totalCost) const;
 	
@@ -384,10 +384,11 @@ public:
 	///  @param[in]		endRef		The reference id of the end polygon.
 	///  @param[in]		startPos	A position within the start polygon. [(x, y, z)]
 	///  @param[in]		endPos		A position within the end polygon. [(x, y, z)]
+	///  @param[in]		maxCost		Cost limit of nodes allowed to be added to the open list	//@UE4
 	///  @param[in]		filter		The polygon filter to apply to the query.
 	/// @returns The status flags for the query.
 	dtStatus initSlicedFindPath(dtPolyRef startRef, dtPolyRef endRef,
-								const float* startPos, const float* endPos,
+								const float* startPos, const float* endPos, const float maxCost, //@UE4
 								const dtQueryFilter* filter);
 
 	/// Updates an in-progress sliced path query.
@@ -780,6 +781,7 @@ private:
 		float lastBestNodeCost;
 		dtPolyRef startRef, endRef;
 		float startPos[3], endPos[3];
+		float maxCost; 					//@UE4 ///< Cost limit of nodes allowed to be added to the open list
 		const dtQueryFilter* filter;
 	};
 	dtQueryData m_query;				///< Sliced query state.
diff --git a/Engine/Source/Runtime/Navmesh/Public/DetourCrowd/DetourPathQueue.h b/Engine/Source/Runtime/Navmesh/Public/DetourCrowd/DetourPathQueue.h
--- a/Engine/Source/Runtime/Navmesh/Public/DetourCrowd/DetourPathQueue.h
+++ b/Engine/Source/Runtime/Navmesh/Public/DetourCrowd/DetourPathQueue.h
@@ -37,6 +37,7 @@ class dtPathQueue
 		/// Path find start and end location.
 		float startPos[3], endPos[3];
 		dtPolyRef startRef, endRef;
+		float maxCost;
 		/// Result.
 		dtPolyRef* path;
 		int npath;
@@ -66,7 +67,7 @@ public:
 	void update(const int maxIters);
 	
 	dtPathQueueRef request(dtPolyRef startRef, dtPolyRef endRef,
-						   const float* startPos, const float* endPos, 
+						   const float* startPos, const float* endPos, const float maxCost,
 						   const dtQueryFilter* filter,
 						   TSharedPtr<dtQuerySpecialLinkFilter> linkFilter);
 	
