Fixed bug where navmesh polygons could connect diagonally across tile boundary.

diff --git a/Engine/Source/Runtime/Navmesh/Private/Detour/DetourNavMesh.cpp b/Engine/Source/Runtime/Navmesh/Private/Detour/DetourNavMesh.cpp
--- a/Engine/Source/Runtime/Navmesh/Private/Detour/DetourNavMesh.cpp
+++ b/Engine/Source/Runtime/Navmesh/Private/Detour/DetourNavMesh.cpp
@@ -37,11 +37,17 @@ inline bool overlapSlabs(const float* amin, const float* amax,
 	// Check for horizontal overlap.
 	// The segment is shrunken a little so that slabs which touch
 	// at end points are not connected.
-	const float minx = dtMax(amin[0]+px,bmin[0]+px);
-	const float maxx = dtMin(amax[0]-px,bmax[0]-px);
-	if (minx > maxx)
+	//@UE4 BEGIN Changed to relative comparison to avoid losing floating point precision.
+	const float minx = dtMax(amin[0], bmin[0]);
+	const float maxx = dtMin(amax[0], bmax[0]);
+	const float diff = maxx - minx;
+	if (diff < px)
+	{
+		*mode = 0; // No overlap
 		return false;
-	
+	}
+	//@UE4 END
+
 	// Check vertical overlap.
 	const float ad = (amax[1]-amin[1]) / (amax[0]-amin[0]);
 	const float ak = amin[1] - ad*amin[0];
@@ -1610,6 +1616,7 @@ int dtNavMesh::queryPolygonsInTile(const dtMeshTile* tile, const float* qmin, co
 /// removed.
 ///
 /// @see dtCreateNavMeshData, #removeTile
+//@UE4 BEGIN
 dtStatus dtNavMesh::addTile(unsigned char* data, int dataSize, int flags,
 							dtTileRef lastRef, dtTileRef* result)
 {
@@ -1742,10 +1749,12 @@ dtStatus dtNavMesh::addTile(unsigned char* data, int dataSize, int flags,
 
 	// Create connections with neighbour tiles.
 	ReadTilesHelper TileArray;
-	int nneis = getTileCountAt(header->x, header->y);
-	dtMeshTile** neis = TileArray.PrepareArray(nneis);
-	
+	int nneis = 0;
+	dtMeshTile** neis = NULL;
+
 	// Connect with layers in current tile.
+	nneis = getTileCountAt(header->x, header->y);
+	neis = TileArray.PrepareArray(nneis);
 	getTilesAt(header->x, header->y, neis, nneis);
 	for (int j = 0; j < nneis; ++j)
 	{
@@ -1764,12 +1773,17 @@ dtStatus dtNavMesh::addTile(unsigned char* data, int dataSize, int flags,
 	{
 		nneis = getNeighbourTilesCountAt(header->x, header->y, i);
 		neis = TileArray.PrepareArray(nneis);
-
 		getNeighbourTilesAt(header->x, header->y, i, neis, nneis);
+
 		for (int j = 0; j < nneis; ++j)
 		{
-			connectExtLinks(tile, neis[j], i, bHasClusters);
-			connectExtLinks(neis[j], tile, dtOppositeTile(i), bHasClusters);
+			// Skip diagonal tiles, nothing to connect there 
+			// (tiles are visited in a ring around the current tile, even tiles are primary directions)
+			if ((i & 1) == 0)
+			{
+				connectExtLinks(tile, neis[j], i, bHasClusters);
+				connectExtLinks(neis[j], tile, dtOppositeTile(i), bHasClusters);
+			}
 			appendSegmentIntersection(segList, tile, neis[j]);
 			connectExtOffMeshLinks(tile, neis[j], i, bHasClusters);
 			connectExtOffMeshLinks(neis[j], tile, dtOppositeTile(i), bHasClusters);
@@ -1784,6 +1798,7 @@ dtStatus dtNavMesh::addTile(unsigned char* data, int dataSize, int flags,
 	
 	return DT_SUCCESS;
 }
+//@UE4 END
 
 const dtMeshTile* dtNavMesh::getTileAt(const int x, const int y, const int layer) const
 {
