trivial fix for CI error

diff --git a/Engine/Source/Runtime/Navmesh/Public/DetourTileCache/DetourTileCacheBuilder.h b/Engine/Source/Runtime/Navmesh/Public/DetourTileCache/DetourTileCacheBuilder.h
--- a/Engine/Source/Runtime/Navmesh/Public/DetourTileCache/DetourTileCacheBuilder.h
+++ b/Engine/Source/Runtime/Navmesh/Public/DetourTileCache/DetourTileCacheBuilder.h
@@ -127,6 +127,8 @@ protected:
 
 struct NAVMESH_API dtTileCacheAlloc
 {
+	virtual ~dtTileCacheAlloc() = default;
+
 	virtual void reset()
 	{
 	}
