Fix navmesh rasterization issue. A wrong offset on the vertical axis caused some voxels to be rasterized outside of their source geometry.
One side effect is that NavModifier aligned wiith geometry woudl not mark the navmesh even when taking into account of the cellheight.

diff --git a/Engine/Source/Runtime/Navmesh/Private/Recast/RecastRasterization.cpp b/Engine/Source/Runtime/Navmesh/Private/Recast/RecastRasterization.cpp
--- a/Engine/Source/Runtime/Navmesh/Private/Recast/RecastRasterization.cpp
+++ b/Engine/Source/Runtime/Navmesh/Private/Recast/RecastRasterization.cpp
@@ -475,8 +475,8 @@ static void rasterizeTri(const float* v0, const float* v1, const float* v2,
 		return;
 	}
 
-	short int triangle_ismin = (short int)rcClamp((int)floorf(triangle_smin * ich + 0.5f), -32000, 32000);
-	short int triangle_ismax = (short int)rcClamp((int)floorf(triangle_smax * ich + 0.5f), -32000, 32000);
+	const short int triangle_ismin = (short int)rcClamp((int)floorf(triangle_smin * ich), -32000, 32000);
+	const short int triangle_ismax = (short int)rcClamp((int)floorf(triangle_smax * ich), -32000, 32000);
 
 	x0 = intMax(x0, 0);
 	int x1_edge = intMin(x1, w);
@@ -628,8 +628,8 @@ static void rasterizeTri(const float* v0, const float* v1, const float* v2,
 		}
 
 		// Snap the span to the heightfield height grid.
-		unsigned short triangle_ismin_clamp = projectToBottom ? 0 : (unsigned short)rcClamp((int)triangle_ismin, 0, RC_SPAN_MAX_HEIGHT); //UE4
-		unsigned short triangle_ismax_clamp = (unsigned short)rcClamp((int)triangle_ismin, (int)triangle_ismin_clamp+1, RC_SPAN_MAX_HEIGHT);
+		const unsigned short triangle_ismin_clamp = projectToBottom ? 0 : (unsigned short)rcClamp((int)triangle_ismin, 0, RC_SPAN_MAX_HEIGHT); //UE4
+		const unsigned short triangle_ismax_clamp = (unsigned short)rcClamp((int)triangle_ismin, (int)triangle_ismin_clamp+1, RC_SPAN_MAX_HEIGHT);
 
 		for (int y = y0; y <= y1; y++)
 		{
@@ -664,7 +664,7 @@ static void rasterizeTri(const float* v0, const float* v1, const float* v2,
 			if (intverts[basevert][0] >= x0 && intverts[basevert][0] <= x1 && intverts[basevert][1] >= y0 && intverts[basevert][1] <= y1)
 			{
 				float sfloat = vertarray[basevert][1] - bmin[1];
-				short int sint = (short int)rcClamp((int)floorf(sfloat * ich + 0.5f), -32000, 32000);
+				short int sint = (short int)rcClamp((int)floorf(sfloat * ich), -32000, 32000);
 	#if TEST_NEW_RASTERIZER
 				rcAssert(sint >= triangle_ismin - 1 && sint <= triangle_ismax + 1);
 	#endif
@@ -702,7 +702,7 @@ static void rasterizeTri(const float* v0, const float* v1, const float* v2,
 					if (y >= y0 && y <= y1)
 					{
 						float sfloat = temppnt[1] - bmin[1];
-						short int sint = (short int)rcClamp((int)floorf(sfloat * ich + 0.5f), -32000, 32000);
+						short int sint = (short int)rcClamp((int)floorf(sfloat * ich), -32000, 32000);
 #if TEST_NEW_RASTERIZER
 						rcAssert(sint >= triangle_ismin - 1 && sint <= triangle_ismax + 1);
 #endif
@@ -743,7 +743,7 @@ static void rasterizeTri(const float* v0, const float* v1, const float* v2,
 						if (x >= x0 && x <= x1)
 						{
 							float sfloat = Inter[i][1] - bmin[1];
-							short int sint = (short int)rcClamp((int)floorf(sfloat * ich + 0.5f), -32000, 32000);
+							short int sint = (short int)rcClamp((int)floorf(sfloat * ich), -32000, 32000);
 #if TEST_NEW_RASTERIZER
 							rcAssert(sint >= triangle_ismin - 1 && sint <= triangle_ismax + 1);
 #endif
@@ -772,7 +772,7 @@ static void rasterizeTri(const float* v0, const float* v1, const float* v2,
 						}
 						for (int x = xloop0; x <= xloop1; x++, sfloat += ds)
 						{
-							short int sint = (short int)rcClamp((int)floorf(sfloat * ich + 0.5f), -32000, 32000);
+							short int sint = (short int)rcClamp((int)floorf(sfloat * ich), -32000, 32000);
 #if TEST_NEW_RASTERIZER
 							rcAssert(sint >= triangle_ismin - 1 && sint <= triangle_ismax + 1);
 #endif
