Navmesh: fix rasterisation area merging when adding new spans
This fixes an issue where navmesh would be wrongly generated on steep slope (as a result, agents would try to reach unreachable locations).

diff --git a/Engine/Source/Runtime/Navmesh/Private/Recast/RecastRasterization.cpp b/Engine/Source/Runtime/Navmesh/Private/Recast/RecastRasterization.cpp
--- a/Engine/Source/Runtime/Navmesh/Private/Recast/RecastRasterization.cpp
+++ b/Engine/Source/Runtime/Navmesh/Private/Recast/RecastRasterization.cpp
@@ -122,15 +122,28 @@ static void addSpan(rcHeightfield& hf, const int x, const int y,
 		}
 		else
 		{
-			// Merge spans.
+			// @UE4 BEGIN
+			// Merge overlapping spans.
+
+			// For spans whose tops are really close to each other, prefer walkable areas.
+			// This is done in order to remove aliasing (similar to z-fighting) on surfaces close to each other.
+			if (rcAbs((int)s->data.smax - (int)cur->data.smax) <= flagMergeThr)
+			{
+				s->data.area = rcMax(s->data.area, cur->data.area);
+			}
+			else
+			{
+				// Use the new spans area if it will become the top.
+				if (cur->data.smax > s->data.smax)
+					s->data.area = cur->data.area;
+			}
+
+			// Merge height intervals.
 			if (cur->data.smin < s->data.smin)
 				s->data.smin = cur->data.smin;
 			if (cur->data.smax > s->data.smax)
 				s->data.smax = cur->data.smax;
-			
-			// Merge flags.
-			if (rcAbs((int)s->data.smax - (int)cur->data.smax) <= flagMergeThr)
-				s->data.area = rcMax(s->data.area, cur->data.area);
+			// @UE4 END
 			
 			// Remove current span.
 			rcSpan* next = cur->next;
