Fix for rare crash in findPath()
Fix result reserve count to prevent unnecessary allocations on the addition of the last item
Limit the number of loop iteration to the path node count

diff --git a/Engine/Source/Runtime/Navmesh/Private/Detour/DetourNavMeshQuery.cpp b/Engine/Source/Runtime/Navmesh/Private/Detour/DetourNavMeshQuery.cpp
--- a/Engine/Source/Runtime/Navmesh/Private/Detour/DetourNavMeshQuery.cpp
+++ b/Engine/Source/Runtime/Navmesh/Private/Detour/DetourNavMeshQuery.cpp
@@ -1633,7 +1633,7 @@ dtStatus dtNavMeshQuery::findPath(dtPolyRef startRef, dtPolyRef endRef,
 	// Reverse the path.
 	dtNode* prev = 0;
 	dtNode* node = lastBestNode;
-	int n = 0;
+	int n = 1;
 	do
 	{
 		dtNode* next = m_nodePool->getNodeAtIdx(node->pidx);
@@ -1661,7 +1661,7 @@ dtStatus dtNavMeshQuery::findPath(dtPolyRef startRef, dtPolyRef endRef,
 
 		node = m_nodePool->getNodeAtIdx(node->pidx);
 	}
-	while (node);
+	while (node && --n > 0);
 
 	if (totalCost)
 	{
