Fix static analysis warning dereferencing null pointer

diff --git a/Engine/Source/Runtime/Navmesh/Private/Recast/Recast.cpp b/Engine/Source/Runtime/Navmesh/Private/Recast/Recast.cpp
--- a/Engine/Source/Runtime/Navmesh/Private/Recast/Recast.cpp
+++ b/Engine/Source/Runtime/Navmesh/Private/Recast/Recast.cpp
@@ -248,6 +248,8 @@ bool rcCreateHeightfield(rcContext* /*ctx*/, rcHeightfield& hf, int width, int h
 
 #if EPIC_ADDITION_USE_NEW_RECAST_RASTERIZER
 	hf.EdgeHits = (rcEdgeHit*)rcAlloc(sizeof(rcEdgeHit) * (hf.height + 1), RC_ALLOC_PERM); 
+	if (!hf.EdgeHits)
+		return false;
 	memset(hf.EdgeHits, 0, sizeof(rcEdgeHit) * (hf.height + 1));
 
 	hf.RowExt = (rcRowExt*)rcAlloc(sizeof(rcRowExt) * (hf.height + 2), RC_ALLOC_PERM); 
@@ -259,6 +261,8 @@ bool rcCreateHeightfield(rcContext* /*ctx*/, rcHeightfield& hf, int width, int h
 	}
 
 	hf.tempspans = (rcTempSpan*)rcAlloc(sizeof(rcTempSpan)*(hf.width + 2) * (hf.height + 2), RC_ALLOC_PERM); 
+	if (!hf.tempspans)
+		return false;
 
 	for (int i = 0; i < hf.height + 2; i++)
 	{
