Tentative fix for SA hf.EdgeHits

diff --git a/Engine/Source/Runtime/Navmesh/Private/Recast/RecastRasterization.cpp b/Engine/Source/Runtime/Navmesh/Private/Recast/RecastRasterization.cpp
--- a/Engine/Source/Runtime/Navmesh/Private/Recast/RecastRasterization.cpp
+++ b/Engine/Source/Runtime/Navmesh/Private/Recast/RecastRasterization.cpp
@@ -442,6 +442,11 @@ static void rasterizeTri(const float* v0, const float* v1, const float* v2,
 						 const int flagMergeThr,
 						 const int rasterizationFlags /*UE4*/)
 {
+	if (!ensure(hf.EdgeHits))
+	{
+		return;
+	}
+
 	const int w = hf.width;
 	const int h = hf.height;
 	const float by = bmax[1] - bmin[1];
