New fix for hf.EdgeHits static analysis warning (SA seems pretty mixed up about it)

diff --git a/Engine/Source/Runtime/Navmesh/Private/Recast/RecastRasterization.cpp b/Engine/Source/Runtime/Navmesh/Private/Recast/RecastRasterization.cpp
--- a/Engine/Source/Runtime/Navmesh/Private/Recast/RecastRasterization.cpp
+++ b/Engine/Source/Runtime/Navmesh/Private/Recast/RecastRasterization.cpp
@@ -442,10 +442,7 @@ static void rasterizeTri(const float* v0, const float* v1, const float* v2,
 						 const int flagMergeThr,
 						 const int rasterizationFlags /*UE4*/)
 {
-	if (!ensure(hf.EdgeHits))
-	{
-		return;
-	}
+	rcEdgeHit* const hfEdgeHits = hf.EdgeHits; //this prevents a static analysis warning
 
 	const int w = hf.width;
 	const int h = hf.height;
@@ -561,8 +558,8 @@ static void rasterizeTri(const float* v0, const float* v1, const float* v2,
 				unsigned char edgeBits = (edge << 4) | (othervert << 2) | basevert;
 				for (int y = loop0; y <= loop1; y++)
 				{
-					int HitIndex = !!hf.EdgeHits[y].Hits[0];
-					hf.EdgeHits[y].Hits[HitIndex] = edgeBits;
+					int HitIndex = !!hfEdgeHits[y].Hits[0];
+					hfEdgeHits[y].Hits[HitIndex] = edgeBits;
 				}
 			}
 			// do the edge intersections with vertical planes
@@ -600,7 +597,7 @@ static void rasterizeTri(const float* v0, const float* v1, const float* v2,
 			float cz = bmin[2] + cs * loop0;
 			for (int y = loop0; y <= loop1; y++, cz += cs)
 			{
-				rcEdgeHit& Hits = hf.EdgeHits[y];
+				rcEdgeHit& Hits = hfEdgeHits[y];
 				if (Hits.Hits[0])
 				{
 					rcAssert(Hits.Hits[1]); // must have two hits
@@ -699,8 +696,8 @@ static void rasterizeTri(const float* v0, const float* v1, const float* v2,
 				unsigned char edgeBits = (edge << 4) | (othervert << 2) | basevert;
 				for (int y = loop0; y <= loop1; y++)
 				{
-					int HitIndex = !!hf.EdgeHits[y].Hits[0];
-					hf.EdgeHits[y].Hits[HitIndex] = edgeBits;
+					int HitIndex = !!hfEdgeHits[y].Hits[0];
+					hfEdgeHits[y].Hits[HitIndex] = edgeBits;
 				}
 			}
 			// do the edge intersections with vertical planes
@@ -743,7 +740,7 @@ static void rasterizeTri(const float* v0, const float* v1, const float* v2,
 			float cz = bmin[2] + cs * loop0;
 			for (int y = loop0; y <= loop1; y++, cz += cs)
 			{
-				rcEdgeHit& Hits = hf.EdgeHits[y];
+				rcEdgeHit& Hits = hfEdgeHits[y];
 				if (Hits.Hits[0])
 				{
 					rcAssert(Hits.Hits[1]); // must have two hits
@@ -874,7 +871,7 @@ static void rasterizeTri(const float* v0, const float* v1, const float* v2,
 	for (int y = 0; y < h; y++)
 	{
 		rcAssert(hf.RowExt[y + 1].MinCol == hf.width + 2 && hf.RowExt[y + 1].MaxCol == -2);
-		rcEdgeHit& Hits = hf.EdgeHits[y];
+		rcEdgeHit& Hits = hfEdgeHits[y];
 		rcAssert(!Hits.Hits[0] && !Hits.Hits[1]);
 		for (int x = 0; x < w; x++)
 		{
@@ -883,7 +880,7 @@ static void rasterizeTri(const float* v0, const float* v1, const float* v2,
 			rcAssert(Temp.sminmax[0] == 32000 && Temp.sminmax[1] == -32000);
 		}
 	}
-	rcEdgeHit& Hits = hf.EdgeHits[h];
+	rcEdgeHit& Hits = hfEdgeHits[h];
 	rcAssert(!Hits.Hits[0] && !Hits.Hits[1]);
 #endif
 }
