Navmesh, Detour -> Move and rotation of navmesh tiles on the tile grid

diff --git a/Engine/Source/Runtime/Navmesh/Private/Detour/DetourCommon.cpp b/Engine/Source/Runtime/Navmesh/Private/Detour/DetourCommon.cpp
--- a/Engine/Source/Runtime/Navmesh/Private/Detour/DetourCommon.cpp
+++ b/Engine/Source/Runtime/Navmesh/Private/Detour/DetourCommon.cpp
@@ -401,6 +401,101 @@ void dtRandomPointInConvexPoly(const float* pts, const int npts, float* areas,
 	out[2] = a*pa[2] + b*pb[2] + c*pc[2];
 }
 
+// @UE4 BEGIN
+dtRotation dtSelectRotation(float rotationDeg)
+{
+	rotationDeg = fmodf(rotationDeg, 360.f);
+	if (rotationDeg < 0)
+		rotationDeg += 360.f;
+
+	// Snap to 90 degrees increment
+	dtRotation rot = DT_ROTATE_0;
+	if (rotationDeg > 45.f && rotationDeg <= 135.f)
+		rot = DT_ROTATE_90;
+	else if (rotationDeg > 135.f && rotationDeg <= 225.f)
+		rot = DT_ROTATE_180;
+	else if (rotationDeg > 225.f && rotationDeg <= 315.f)
+		rot = DT_ROTATE_270;
+
+	return rot;
+}
+
+void dtVRot90(float* dest, const float* v, const dtRotation rot)
+{
+	dest[1] = v[1];
+
+	switch (rot)
+	{
+	case DT_ROTATE_90:
+		dest[0] = -v[2];
+		dest[2] = v[0];
+		break;
+	case DT_ROTATE_180:
+		dest[0] = -v[0];
+		dest[2] = -v[2];
+		break;
+	case DT_ROTATE_270:
+		dest[0] = v[2];
+		dest[2] = -v[0];
+		break;
+	default:
+		// DT_ROTATE_0
+		dest[0] = v[0];
+		dest[2] = v[2];
+		break;
+	}
+}
+
+void dtVRot90(unsigned short* dest, const unsigned short* v, const dtRotation rot)
+{
+	dest[1] = v[1];
+
+	switch (rot)
+	{
+	case DT_ROTATE_90:
+		dest[0] = -v[2];
+		dest[2] = v[0];
+		break;
+	case DT_ROTATE_180:
+		dest[0] = -v[0];
+		dest[2] = -v[2];
+		break;
+	case DT_ROTATE_270:
+		dest[0] = v[2];
+		dest[2] = -v[0];
+		break;
+	default:
+		// DT_ROTATE_0
+		dest[0] = v[0];
+		dest[2] = v[2];
+		break;
+	}
+}
+
+void dtRotate90(float* dest, const float* v, const float* center, const dtRotation rot)
+{
+	float localPos[3];
+	dtVsub(localPos, v, center);
+	float newLocalPos[3];
+	dtVRot90(newLocalPos, localPos, rot);
+	dtVadd(dest, center, newLocalPos);
+}
+
+void dtRotate90(unsigned short* dest, const unsigned short* v, const unsigned short* center, const dtRotation rot)
+{
+	unsigned short localPos[3];
+	localPos[0] = v[0] - center[0];
+	localPos[2] = v[2] - center[2];
+
+	unsigned short newLocalPos[3];
+	dtVRot90(newLocalPos, localPos, rot);
+
+	dest[0] = center[0] + newLocalPos[0];
+	dest[1] = v[1];
+	dest[2] = center[2] + newLocalPos[2];
+}
+// @UE4 END
+
 inline float vperpXZ(const float* a, const float* b) { return a[0]*b[2] - a[2]*b[0]; }
 
 bool dtIntersectSegSeg2D(const float* ap, const float* aq,
diff --git a/Engine/Source/Runtime/Navmesh/Private/Detour/DetourNavMeshBuilder.cpp b/Engine/Source/Runtime/Navmesh/Private/Detour/DetourNavMeshBuilder.cpp
--- a/Engine/Source/Runtime/Navmesh/Private/Detour/DetourNavMeshBuilder.cpp
+++ b/Engine/Source/Runtime/Navmesh/Private/Detour/DetourNavMeshBuilder.cpp
@@ -947,3 +947,187 @@ bool dtNavMeshDataSwapEndian(unsigned char* data, const int /*dataSize*/)
  
 	return true;
 }
+
+// @UE4 BEGIN
+bool dtTransformTileData(unsigned char* data, const int dataSize, const int offsetX, const int offsetY, const float tileWidth, const float tileHeight, const float rotationDeg)
+{
+	// Make sure the data is in right format.
+	dtMeshHeader* header = (dtMeshHeader*)data;
+	if (header->magic != DT_NAVMESH_MAGIC)
+		return false;
+	if (header->version != DT_NAVMESH_VERSION)
+		return false;
+
+	// Set new coords
+	header->x += offsetX;
+	header->y += offsetY;
+
+	// Patch header pointers.
+	const int headerSize = dtAlign4(sizeof(dtMeshHeader));
+	const int vertsSize = dtAlign4(sizeof(float)*3*header->vertCount);
+	const int polysSize = dtAlign4(sizeof(dtPoly)*header->polyCount);
+	const int linksSize = dtAlign4(sizeof(dtLink)*(header->maxLinkCount));
+	const int detailMeshesSize = dtAlign4(sizeof(dtPolyDetail)*header->detailMeshCount);
+	const int detailVertsSize = dtAlign4(sizeof(float)*3*header->detailVertCount);
+	const int detailTrisSize = dtAlign4(sizeof(unsigned char)*4*header->detailTriCount);
+	const int bvtreeSize = dtAlign4(sizeof(dtBVNode)*header->bvNodeCount);
+	const int offMeshLinksSize = dtAlign4(sizeof(dtOffMeshConnection)*header->offMeshConCount);
+	const int offMeshSegsSize = dtAlign4(sizeof(dtOffMeshSegmentConnection)*header->offMeshSegConCount);
+	const int clustersSize = dtAlign4(sizeof(dtCluster)*header->clusterCount);
+	const int clusterPolysSize = dtAlign4(sizeof(unsigned short)*header->offMeshBase);
+
+	dtMeshTile tile;
+	tile.header = header;
+
+	unsigned char* d = data + headerSize;
+	tile.verts = (float*)d; d += vertsSize;
+	tile.polys = (dtPoly*)d; d += polysSize;
+	tile.links = (dtLink*)d; d += linksSize;
+	tile.detailMeshes = (dtPolyDetail*)d; d += detailMeshesSize;
+	tile.detailVerts = (float*)d; d += detailVertsSize;
+	tile.detailTris = (unsigned char*)d; d += detailTrisSize;
+	tile.bvTree = (dtBVNode*)d; d += bvtreeSize;
+	tile.offMeshCons = (dtOffMeshConnection*)d; d += offMeshLinksSize;
+	tile.offMeshSeg = (dtOffMeshSegmentConnection*)d; d += offMeshSegsSize;
+	tile.clusters = (dtCluster*)d; d += clustersSize;
+	tile.polyClusters = (unsigned short*)d; d += clusterPolysSize;
+
+	// Rotate on original center position
+	float rotationCenter[3];
+	dtVadd(rotationCenter, header->bmin, header->bmax);
+	dtVscale(rotationCenter, rotationCenter, 0.5f);
+
+	// Compute offset
+	float offset[3];
+	offset[0] = tileWidth * offsetX;
+	offset[1] = 0.f;
+	offset[2] = tileHeight * offsetY;
+
+	// Compute center for bvtree
+	const float qfac = tile.header->bvQuantFactor;
+	unsigned short qMin[3];
+	unsigned short qMax[3];
+	qMin[0] = (unsigned short)(qfac * tile.header->bmin[0]);
+	qMin[1] = (unsigned short)(qfac * tile.header->bmin[1]);
+	qMin[2] = (unsigned short)(qfac * tile.header->bmin[2]);
+	qMax[0] = (unsigned short)(qfac * tile.header->bmax[0]);
+	qMax[1] = (unsigned short)(qfac * tile.header->bmax[1]);
+	qMax[2] = (unsigned short)(qfac * tile.header->bmax[2]);
+	unsigned short qLocalCenter[3];
+	qLocalCenter[0] = (qMax[0] - qMin[0])/2;
+	qLocalCenter[1] = (qMax[1] - qMin[1])/2;
+	qLocalCenter[2] = (qMax[2] - qMin[2])/2;
+
+	const dtRotation rot = dtSelectRotation(rotationDeg);
+
+	// Transform tile bounds
+	float tmin[3];
+	float tmax[3];
+	dtRotate90(tmin, tile.header->bmin, rotationCenter, rot);
+	dtVadd(tmin, tmin, offset);
+	dtRotate90(tmax, tile.header->bmax, rotationCenter, rot);
+	dtVadd(tmax, tmax, offset);
+	dtVcopy(tile.header->bmin, tmin);
+	dtVcopy(tile.header->bmax, tmax);
+	dtVmin(tile.header->bmin, tmax);
+	dtVmax(tile.header->bmax, tmin);
+
+	// Update sides
+	for (int j = 0; j < tile.header->polyCount; ++j)
+	{
+		dtPoly* poly = &tile.polys[j];
+		const int nv = poly->vertCount;
+		for (int vi = 0; vi < nv; ++vi)
+		{
+			// Skip non-portal edges.
+			if ((poly->neis[vi] & DT_EXT_LINK) == 0)
+				continue;
+
+			// Rotate
+			// 3 2 1
+			// 4   0
+			// 5 6 7
+			const unsigned short sideMask = 0xff;
+			const unsigned short side = poly->neis[vi] & sideMask;
+			const unsigned short newSide = (side + (2*rot)) % 8; //rot [0..3], newSide [0,2,4,6]
+			poly->neis[vi] = (poly->neis[vi] & 0xff00) | newSide;
+		}
+	}
+
+	// Transform tile vertices
+	for (int j = 0; j < tile.header->vertCount; ++j)
+	{
+		dtRotate90(&(tile.verts[j*3]), &(tile.verts[j*3]), rotationCenter, rot);
+		dtVadd(&(tile.verts[j*3]), &(tile.verts[j*3]), offset);
+	}
+
+	// Transform tile details vertices
+	for (int j = 0; j < tile.header->detailVertCount; ++j)
+	{
+		dtRotate90(&(tile.detailVerts[j*3]), &(tile.detailVerts[j*3]), rotationCenter, rot);
+		dtVadd(&(tile.detailVerts[j*3]), &(tile.detailVerts[j*3]), offset);
+	}
+
+	// Transfrom BVTree (bmin and bmax are local to the tile)
+	for (int j = 0; j < tile.header->bvNodeCount; ++j)
+	{
+		dtBVNode* node = &tile.bvTree[j];
+		unsigned short min[3];
+		unsigned short max[3];
+		dtRotate90(min, node->bmin, qLocalCenter, rot);
+		dtRotate90(max, node->bmax, qLocalCenter, rot);
+		node->bmin[0] = dtMin(min[0], max[0]);
+		node->bmin[1] = dtMin(min[1], max[1]);
+		node->bmin[2] = dtMin(min[2], max[2]);
+		node->bmax[0] = dtMax(min[0], max[0]);
+		node->bmax[1] = dtMax(min[1], max[1]);
+		node->bmax[2] = dtMax(min[2], max[2]);
+	}
+
+	// Transform off-mesh connections
+	for (int j = 0; j < tile.header->offMeshConCount; ++j)
+	{
+		dtRotate90(&(tile.offMeshCons[j].pos[0]), &(tile.offMeshCons[j].pos[0]), rotationCenter, rot);
+		dtRotate90(&(tile.offMeshCons[j].pos[3]), &(tile.offMeshCons[j].pos[3]), rotationCenter, rot);
+		dtVadd(&(tile.offMeshCons[j].pos[0]), &(tile.offMeshCons[j].pos[0]), offset);
+		dtVadd(&(tile.offMeshCons[j].pos[3]), &(tile.offMeshCons[j].pos[3]), offset);
+
+		tile.offMeshCons[j].side = (tile.offMeshCons[j].side + (2*rot)) % 8; //rot [0..3], side [0,1,2,3,4,5,6,7]
+	}
+
+	// Transform off-mesh segment connections
+	for (int j = 0; j < tile.header->offMeshSegConCount; ++j)
+	{
+		dtRotate90(&(tile.offMeshSeg[j].startA[0]), &(tile.offMeshSeg[j].startA[0]), rotationCenter, rot);
+		dtRotate90(&(tile.offMeshSeg[j].endA[0]), &(tile.offMeshSeg[j].endA[0]), rotationCenter, rot);
+		dtRotate90(&(tile.offMeshSeg[j].startB[0]), &(tile.offMeshSeg[j].startB[0]), rotationCenter, rot);
+		dtRotate90(&(tile.offMeshSeg[j].endB[0]), &(tile.offMeshSeg[j].endB[0]), rotationCenter, rot);
+		dtVadd(&(tile.offMeshSeg[j].startA[0]), &(tile.offMeshSeg[j].startA[0]), offset);
+		dtVadd(&(tile.offMeshSeg[j].endA[0]), &(tile.offMeshSeg[j].endA[0]), offset);
+		dtVadd(&(tile.offMeshSeg[j].startB[0]), &(tile.offMeshSeg[j].startB[0]), offset);
+		dtVadd(&(tile.offMeshSeg[j].endB[0]), &(tile.offMeshSeg[j].endB[0]), offset);
+	}
+
+	// Transform clusters
+	for (int j = 0; j < tile.header->clusterCount; ++j)
+	{
+		dtRotate90(&(tile.clusters[j].center[0]), &(tile.clusters[j].center[0]), rotationCenter, rot);
+		dtVadd(&(tile.clusters[j].center[0]), &(tile.clusters[j].center[0]), offset);
+	}
+
+	return true;
+}
+
+void dtComputeTileOffsetFromRotation(const float* position, const float* rotationCenter, const float rotationDeg, const float tileWidth, const float tileHeight, int& deltaX, int& deltaY)
+{
+	float relativeTilePos[3];
+	dtVsub(relativeTilePos, position, rotationCenter);
+	float newRelativeTilePos[3] = { 0.f, 0.f, 0.f };
+	dtVRot90(newRelativeTilePos, relativeTilePos, dtSelectRotation(rotationDeg));
+
+	float RcTilePosChangedVector[3];
+	dtVsub(RcTilePosChangedVector, newRelativeTilePos, relativeTilePos);
+	deltaX = RcTilePosChangedVector[0] / tileWidth;
+	deltaY = RcTilePosChangedVector[2] / tileHeight;
+}
+// @UE4 END
diff --git a/Engine/Source/Runtime/Navmesh/Public/Detour/DetourCommon.h b/Engine/Source/Runtime/Navmesh/Public/Detour/DetourCommon.h
--- a/Engine/Source/Runtime/Navmesh/Public/Detour/DetourCommon.h
+++ b/Engine/Source/Runtime/Navmesh/Public/Detour/DetourCommon.h
@@ -486,6 +486,47 @@ inline void dtSwapEndian(float* v)
 void dtRandomPointInConvexPoly(const float* pts, const int npts, float* areas,
 							   const float s, const float t, float* out);
 
+// @UE4 BEGIN
+enum dtRotation
+{
+	DT_ROTATE_0, 
+	DT_ROTATE_90, 
+	DT_ROTATE_180, 
+	DT_ROTATE_270
+};
+
+/// Select a 90 degree increment value from an input angle in degree.
+///  @param[in]		rotationDeg The desired rotation in degree.
+///  @return The rotation enum value.
+dtRotation dtSelectRotation(float rotationDeg);
+
+/// Rotate by 90 degree increments.
+///  @param[out]	dest	The result position. [(x, y, z)]
+///  @param[in]		v		The vector to rotate. [(x, y, z)]
+///  @param[in]		rot		The rotation enum value.
+void dtVRot90(float* dest, const float* v, const dtRotation rot);
+
+/// Rotate by 90 degree increments.
+///  @param[out]	dest	The result position. [(x, y, z)]
+///  @param[in]		v		The vector to rotate. [(x, y, z)]
+///  @param[in]		rot		The rotation enum value.
+void dtVRot90(unsigned short* dest, const unsigned short* v, const dtRotation rot);
+
+/// Rotate vector around center position by increments of 90 degrees.
+///  @param[out]	dest	The result position. [(x, y, z)]
+///  @param[in]		v		The vector to rotate. [(x, y, z)]
+///  @param[in]		center	The center point. [(x, y, z)]
+///  @param[in]		rot		The rotation enum value.
+void dtRotate90(float* dest, const float* v, const float* center, const dtRotation rot);
+
+/// Rotate vector around center position by increments of 90 degrees.
+///  @param[out]	dest	The result position. [(x, y, z)]
+///  @param[in]		v		The vector to rotate. [(x, y, z)]
+///  @param[in]		center	The center point. [(x, y, z)]
+///  @param[in]		rot		The rotation enum value.
+void dtRotate90(unsigned short* dest, const unsigned short* v, const unsigned short* center, const dtRotation rot);
+// @UE4 END
+
 /// @}
 
 #endif // DETOURCOMMON_H
diff --git a/Engine/Source/Runtime/Navmesh/Public/Detour/DetourNavMeshBuilder.h b/Engine/Source/Runtime/Navmesh/Public/Detour/DetourNavMeshBuilder.h
--- a/Engine/Source/Runtime/Navmesh/Public/Detour/DetourNavMeshBuilder.h
+++ b/Engine/Source/Runtime/Navmesh/Public/Detour/DetourNavMeshBuilder.h
@@ -162,6 +162,28 @@ NAVMESH_API bool dtNavMeshHeaderSwapEndian(unsigned char* data, const int dataSi
 ///  @param[in]		dataSize	The size of the data array.
 NAVMESH_API bool dtNavMeshDataSwapEndian(unsigned char* data, const int dataSize);
 
+// @UE4 BEGIN
+/// Offset and rotate around center the data in the tile
+///  @param[in,out]	data			Data of the tile mesh. (See: #dtCreateNavMeshData)
+///  @param[in]		dataSize		Data size of the tile mesh.
+///  @param[in]		offsetX			X offset in tile coordinates.
+///  @param[in]		offsetY			Y offset in tile coordinates.
+///  @param[in]		tileWidth		Tile width.
+///  @param[in]		tileHeight		Tile height.
+///  @param[in]		rotationDeg		Rotation in degrees.
+NAVMESH_API bool dtTransformTileData(unsigned char* data, const int dataSize, const int offsetX, const int offsetY, const float tileWidth, const float tileHeight, const float rotationDeg);
+
+/// Compute XY offset caused by the given rotation
+///  @param[in]		position		Position to rotate. [(x, y, z)]
+///  @param[in]		rotationCenter	Rotation center. [(x, y, z)]
+///  @param[in]		rotationDeg		Rotation in degrees.
+///  @param[in]		tileWidth		Tile width.
+///  @param[in]		tileHeight		Tile height.
+///  @param[out]	deltaX			Offset X in tile coordinates.
+///  @param[out]	deltaY			Offset Y in tile coordinates.
+NAVMESH_API void dtComputeTileOffsetFromRotation(const float* position, const float* rotationCenter, const float rotationDeg, const float tileWidth, const float tileHeight, int& deltaX, int& deltaY);
+// @UE4 END
+
 #endif // DETOURNAVMESHBUILDER_H
 
 // This section contains detailed documentation for members that don't have
