Updated occurences of "UE4" in Navmesh module's comments

diff --git a/Engine/Source/Runtime/Navmesh/Private/DebugUtils/DebugDraw.cpp b/Engine/Source/Runtime/Navmesh/Private/DebugUtils/DebugDraw.cpp
--- a/Engine/Source/Runtime/Navmesh/Private/DebugUtils/DebugDraw.cpp
+++ b/Engine/Source/Runtime/Navmesh/Private/DebugUtils/DebugDraw.cpp
@@ -20,7 +20,7 @@
 //
 
 #include "DebugUtils/DebugDraw.h"
-#include "Recast/RecastAlloc.h" // UE4
+#include "Recast/RecastAlloc.h" // UE
 #define _USE_MATH_DEFINES
 
 
@@ -575,22 +575,22 @@ duDisplayList::duDisplayList(int cap) :
 
 duDisplayList::~duDisplayList()
 {
-	rcFree(m_pos); // UE4
-	rcFree(m_color); // UE4
+	rcFree(m_pos); // UE
+	rcFree(m_color); // UE
 }
 
 void duDisplayList::resize(int cap)
 {
-	float* newPos = (float*)rcAlloc(sizeof(float)*3*cap, RC_ALLOC_PERM); // UE4
+	float* newPos = (float*)rcAlloc(sizeof(float)*3*cap, RC_ALLOC_PERM); // UE
 	if (m_size)
 		memcpy(newPos, m_pos, sizeof(float)*3*m_size);
-	rcFree(m_pos); // UE4
+	rcFree(m_pos); // UE
 	m_pos = newPos;
 
-	unsigned int* newColor = (unsigned int*)rcAlloc(sizeof(unsigned int)*cap, RC_ALLOC_PERM); // UE4
+	unsigned int* newColor = (unsigned int*)rcAlloc(sizeof(unsigned int)*cap, RC_ALLOC_PERM); // UE
 	if (m_size)
 		memcpy(newColor, m_color, sizeof(unsigned int)*m_size);
-	rcFree(m_color); // UE4
+	rcFree(m_color); // UE
 	m_color = newColor;
 	
 	m_cap = cap;
diff --git a/Engine/Source/Runtime/Navmesh/Private/DebugUtils/DetourDebugDraw.cpp b/Engine/Source/Runtime/Navmesh/Private/DebugUtils/DetourDebugDraw.cpp
--- a/Engine/Source/Runtime/Navmesh/Private/DebugUtils/DetourDebugDraw.cpp
+++ b/Engine/Source/Runtime/Navmesh/Private/DebugUtils/DetourDebugDraw.cpp
@@ -243,7 +243,7 @@ static void drawMeshTile(duDebugDraw* dd, const dtNavMesh& mesh, const dtNavMesh
 		}
 		dd->end();
 
-//@UE4 BEGIN
+//@UE BEGIN
 #if WITH_NAVMESH_SEGMENT_LINKS
 		dd->begin(DU_DRAW_LINES, 4.0f);
 		for (int i = 0; i < tile->header->polyCount; ++i)
@@ -284,7 +284,7 @@ static void drawMeshTile(duDebugDraw* dd, const dtNavMesh& mesh, const dtNavMesh
 		}
 		dd->end();
 #endif // WITH_NAVMESH_SEGMENT_LINKS
-//@UE4 END
+//@UE END
 	}
 	
 	const unsigned int vcol = duRGBA(0,0,0,196);
@@ -527,7 +527,7 @@ void duDebugDrawNavMeshPoly(duDebugDraw* dd, const dtNavMesh& mesh, dtPolyRef re
 		
 		dd->end();
 	}
-//@UE4 BEGIN
+//@UE BEGIN
 #if WITH_NAVMESH_SEGMENT_LINKS
 	else if (poly->getType() == DT_POLYTYPE_OFFMESH_SEGMENT)
 	{
@@ -540,7 +540,7 @@ void duDebugDrawNavMeshPoly(duDebugDraw* dd, const dtNavMesh& mesh, dtPolyRef re
 		dd->end();
 	}
 #endif // WITH_NAVMESH_SEGMENT_LINKS
-//@UE4 END
+//@UE END
 	else if (poly->getType() == DT_POLYTYPE_GROUND)
 	{
 		const dtPolyDetail* pd = &tile->detailMeshes[ip];
@@ -564,7 +564,7 @@ void duDebugDrawNavMeshPoly(duDebugDraw* dd, const dtNavMesh& mesh, dtPolyRef re
 
 }
 
-//@UE4 BEGIN
+//@UE BEGIN
 #if WITH_NAVMESH_CLUSTER_LINKS
 void duDebugDrawNavMeshClusters(struct duDebugDraw* dd, const dtNavMesh& mesh)
 {
@@ -601,7 +601,7 @@ void duDebugDrawNavMeshCluster(struct duDebugDraw* dd, const dtNavMesh& mesh, dt
 	}
 }
 #endif // WITH_NAVMESH_CLUSTER_LINKS
-//@UE4 END
+//@UE END
 
 static void debugDrawTileCachePortals(struct duDebugDraw* dd, const dtTileCacheLayer& layer, const float cs, const float ch)
 {
diff --git a/Engine/Source/Runtime/Navmesh/Private/Detour/DetourCommon.cpp b/Engine/Source/Runtime/Navmesh/Private/Detour/DetourCommon.cpp
--- a/Engine/Source/Runtime/Navmesh/Private/Detour/DetourCommon.cpp
+++ b/Engine/Source/Runtime/Navmesh/Private/Detour/DetourCommon.cpp
@@ -401,7 +401,7 @@ void dtRandomPointInConvexPoly(const float* pts, const int npts, float* areas,
 	out[2] = a*pa[2] + b*pb[2] + c*pc[2];
 }
 
-// @UE4 BEGIN
+// @UE BEGIN
 dtRotation dtSelectRotation(float rotationDeg)
 {
 	rotationDeg = fmodf(rotationDeg, 360.f);
@@ -494,7 +494,7 @@ void dtRotate90(unsigned short* dest, const unsigned short* v, const unsigned sh
 	dest[1] = v[1];
 	dest[2] = center[2] + newLocalPos[2];
 }
-// @UE4 END
+// @UE END
 
 inline float vperpXZ(const float* a, const float* b) { return a[0]*b[2] - a[2]*b[0]; }
 
diff --git a/Engine/Source/Runtime/Navmesh/Private/Detour/DetourNavMesh.cpp b/Engine/Source/Runtime/Navmesh/Private/Detour/DetourNavMesh.cpp
--- a/Engine/Source/Runtime/Navmesh/Private/Detour/DetourNavMesh.cpp
+++ b/Engine/Source/Runtime/Navmesh/Private/Detour/DetourNavMesh.cpp
@@ -37,7 +37,7 @@ inline bool overlapSlabs(const float* amin, const float* amax,
 	// Check for horizontal overlap.
 	// The segment is shrunken a little so that slabs which touch
 	// at end points are not connected.
-	//@UE4 BEGIN Changed to relative comparison to avoid losing floating point precision.
+	//@UE BEGIN Changed to relative comparison to avoid losing floating point precision.
 	const float minx = dtMax(amin[0], bmin[0]);
 	const float maxx = dtMin(amax[0], bmax[0]);
 	const float diff = maxx - minx;
@@ -46,7 +46,7 @@ inline bool overlapSlabs(const float* amin, const float* amax,
 		*mode = 0; // No overlap
 		return false;
 	}
-	//@UE4 END
+	//@UE END
 
 	// Check vertical overlap.
 	const float ad = (amax[1]-amin[1]) / (amax[0]-amin[0]);
@@ -165,7 +165,7 @@ inline int computeTileHash(int x, int y, const int mask)
 	return (int)(n & mask);
 }
 
-//@UE4 BEGIN
+//@UE BEGIN
 enum ELinkAllocationType
 {
 	// allocated in linksFreeList from dtMeshTile (detour)
@@ -176,7 +176,7 @@ enum ELinkAllocationType
 	CREATE_LINK_DYNAMIC_CLUSTER,
 #endif // WITH_NAVMESH_SEGMENT_LINKS
 };
-//@UE4 END
+//@UE END
 
 inline unsigned int allocLink(dtMeshTile* tile, char LinkAllocMode)
 {
@@ -190,7 +190,7 @@ inline unsigned int allocLink(dtMeshTile* tile, char LinkAllocMode)
 			tile->linksFreeList = tile->links[newLink].next;
 		}
 	}
-//@UE4 BEGIN: offmesh links will be added in dynamic array. 
+//@UE BEGIN: offmesh links will be added in dynamic array. 
 	// Both Point to Point (stock detour) and Segment to Segment (unreal)
 	else if (LinkAllocMode == CREATE_LINK_DYNAMIC_OFFMESH)
 	{
@@ -228,14 +228,14 @@ inline unsigned int allocLink(dtMeshTile* tile, char LinkAllocMode)
 		newLink += DT_CLINK_FIRST;
 	}
 #endif // WITH_NAVMESH_CLUSTER_LINKS
-//@UE4 END
+//@UE END
 
 	return newLink;
 }
 
 inline void freeLink(dtMeshTile* tile, unsigned int link)
 {
-	//@UE4 BEGIN: offmesh links were added in dynamic array
+	//@UE BEGIN: offmesh links were added in dynamic array
 #if WITH_NAVMESH_CLUSTER_LINKS
 	const unsigned int firstClusterLinkIdx = DT_CLINK_FIRST;
 #else
@@ -261,7 +261,7 @@ inline void freeLink(dtMeshTile* tile, unsigned int link)
 		tile->dynamicFreeListC = linkIdx;
 	}
 #endif // WITH_NAVMESH_CLUSTER_LINKS
-	//@UE4 END
+	//@UE END
 }
 
 dtNavMesh* dtAllocNavMesh()
@@ -282,7 +282,7 @@ void dtFreeNavMesh(dtNavMesh* navmesh)
 	dtFree(navmesh);
 }
 
-//@UE4 BEGIN
+//@UE BEGIN
 void dtFreeNavMeshTileRuntimeData(dtMeshTile* tile)
 {
 	tile->dynamicLinksO.~dtChunkArray();
@@ -710,7 +710,7 @@ static void createSegmentLinks(dtNavMesh* nav, dtOffMeshSegmentData* seg, dtMesh
 	}
 }
 #endif // WITH_NAVMESH_SEGMENT_LINKS
-//@UE4 END
+//@UE END
 
 //////////////////////////////////////////////////////////////////////////////////////////
 
@@ -944,9 +944,9 @@ void dtNavMesh::unconnectExtLinks(dtMeshTile* tile, dtMeshTile* target)
 		while (j != DT_NULL_LINK)
 		{
 			dtLink& testLink = getLink(tile, j);
-//@UE4 BEGIN
+//@UE BEGIN
 			if ((testLink.side & DT_CONNECTION_INTERNAL) == 0 &&
-//@UE4 END
+//@UE END
 				decodePolyIdTile(testLink.ref) == targetNum)
 			{
 				// Remove link.
@@ -972,11 +972,11 @@ void dtNavMesh::unconnectExtLinks(dtMeshTile* tile, dtMeshTile* target)
 		}
 	}
 
-//@UE4 BEGIN
+//@UE BEGIN
 #if WITH_NAVMESH_CLUSTER_LINKS
 	unconnectClusterLinks(tile, target);
 #endif // WITH_NAVMESH_CLUSTER_LINKS
-//@UE4 END
+//@UE END
 }
 
 void dtNavMesh::connectExtLinks(dtMeshTile* tile, dtMeshTile* target, int side, bool updateCLinks)
@@ -1047,7 +1047,7 @@ void dtNavMesh::connectExtLinks(dtMeshTile* tile, dtMeshTile* target, int side,
 					}
 				}
 
-				//@UE4 BEGIN
+				//@UE BEGIN
 #if WITH_NAVMESH_CLUSTER_LINKS
 				if (updateCLinks)
 				{
@@ -1061,13 +1061,13 @@ void dtNavMesh::connectExtLinks(dtMeshTile* tile, dtMeshTile* target, int side,
 					}
 				}
 #endif // WITH_NAVMESH_CLUSTER_LINKS
-				//@UE4 END
+				//@UE END
 			}
 		}
 	}
 }
 
-//@UE4 BEGIN
+//@UE BEGIN
 void dtNavMesh::linkOffMeshHelper(dtMeshTile* tile0, unsigned int polyIdx0, dtMeshTile* tile1, unsigned int polyIdx1, unsigned char side, unsigned char edge)
 {
 	dtPoly* poly0 = &tile0->polys[polyIdx0];
@@ -1082,7 +1082,7 @@ void dtNavMesh::linkOffMeshHelper(dtMeshTile* tile0, unsigned int polyIdx0, dtMe
 	link.next = poly0->firstLink;
 	poly0->firstLink = idx;
 }
-//@UE4 END
+//@UE END
 
 void dtNavMesh::connectExtOffMeshLinks(dtMeshTile* tile, dtMeshTile* target, int side, bool updateCLinks)
 {
@@ -1090,7 +1090,7 @@ void dtNavMesh::connectExtOffMeshLinks(dtMeshTile* tile, dtMeshTile* target, int
 	
 	// Connect off-mesh links.
 	// We are interested on links which land from target tile to this tile.
-	//@UE4 BEGIN
+	//@UE BEGIN
 	const unsigned char oppositeSide = (side == -1) ? DT_CONNECTION_INTERNAL : (unsigned char)dtOppositeTile(side);
 
 	for (int i = 0; i < target->header->offMeshConCount; ++i)
@@ -1115,7 +1115,7 @@ void dtNavMesh::connectExtOffMeshLinks(dtMeshTile* tile, dtMeshTile* target, int
 		float nearestPt[3];
 		dtPolyRef ref = 0;
 
-		// [UE4] try finding cheapest, but it that's outside requested radius, fallback to nearest one
+		// [UE] try finding cheapest, but it that's outside requested radius, fallback to nearest one
 		// findNearestPoly may return too optimistic results, further check to make sure. 
 		if (targetCon->getSnapToCheapestArea())
 		{
@@ -1190,7 +1190,7 @@ void dtNavMesh::connectExtOffMeshLinks(dtMeshTile* tile, dtMeshTile* target, int
 		}
 #endif // WITH_NAVMESH_CLUSTER_LINKS
 	}
-	//@UE4 END
+	//@UE END
 }
 
 void dtNavMesh::connectIntLinks(dtMeshTile* tile)
@@ -1220,9 +1220,9 @@ void dtNavMesh::connectIntLinks(dtMeshTile* tile)
 				dtLink* link = &tile->links[idx];
 				link->ref = base | (dtPolyRef)(poly->neis[j]-1);
 				link->edge = (unsigned char)j;
-//@UE4 BEGIN 
+//@UE BEGIN 
 				link->side = DT_CONNECTION_INTERNAL;
-//@UE4 END
+//@UE END
 				link->bmin = link->bmax = 0;
 				// Add to linked list.
 				link->next = poly->firstLink;
@@ -1249,7 +1249,7 @@ void dtNavMesh::baseOffMeshLinks(dtMeshTile* tile)
 		float nearestPt[3];
 		dtPolyRef ref = 0;
 
-		// [UE4] try finding cheapest, but it that's outside requested radius, fallback to nearest one
+		// [UE] try finding cheapest, but it that's outside requested radius, fallback to nearest one
 		// findNearestPoly may return too optimistic results, further check to make sure. 
 		if (con->getSnapToCheapestArea())
 		{
@@ -1286,7 +1286,7 @@ void dtNavMesh::baseOffMeshLinks(dtMeshTile* tile)
 	}
 }
 
-//@UE4 BEGIN
+//@UE BEGIN
 #if WITH_NAVMESH_CLUSTER_LINKS
 void dtNavMesh::connectClusterLink(dtMeshTile* tile0, unsigned int clusterIdx0, dtMeshTile* tile1, unsigned int clusterIdx1, unsigned char flags, bool bCheckExisting)
 {
@@ -1370,7 +1370,7 @@ void dtNavMesh::unconnectClusterLinks(dtMeshTile* tile0, dtMeshTile* tile1)
 	}
 }
 #endif // WITH_NAVMESH_CLUSTER_LINKS
-//@UE4 END
+//@UE END
 
 void dtNavMesh::closestPointOnPolyInTile(const dtMeshTile* tile, unsigned int ip,
 										 const float* pos, float* closest) const
@@ -1658,7 +1658,7 @@ int dtNavMesh::queryPolygonsInTile(const dtMeshTile* tile, const float* qmin, co
 /// removed.
 ///
 /// @see dtCreateNavMeshData, #removeTile
-//@UE4 BEGIN
+//@UE BEGIN
 dtStatus dtNavMesh::addTile(unsigned char* data, int dataSize, int flags,
 							dtTileRef lastRef, dtTileRef* result)
 {
@@ -1866,7 +1866,7 @@ dtStatus dtNavMesh::addTile(unsigned char* data, int dataSize, int flags,
 	
 	return DT_SUCCESS;
 }
-//@UE4 END
+//@UE END
 
 const dtMeshTile* dtNavMesh::getTileAt(const int x, const int y, const int layer) const
 {
@@ -1905,7 +1905,7 @@ int dtNavMesh::getNeighbourTilesAt(const int x, const int y, const int side, dtM
 	return getTilesAt(nx, ny, (const dtMeshTile**)tiles, maxTiles);
 }
 
-// @UE4 BEGIN
+// @UE BEGIN
 int dtNavMesh::getNeighbourTilesCountAt(const int x, const int y, const int side) const
 {
 	int nx = x, ny = y;
@@ -1944,7 +1944,7 @@ int dtNavMesh::getTileCountAt(const int x, const int y) const
 
 	return n;
 }
-// @UE4 END
+// @UE END
 
 int dtNavMesh::getTilesAt(const int x, const int y, dtMeshTile** tiles, const int maxTiles) const
 {
@@ -2360,7 +2360,7 @@ dtStatus dtNavMesh::getOffMeshConnectionPolyEndPoints(dtPolyRef prevRef, dtPolyR
 		i = link.next;
 	}
 	
-	//@UE4 BEGIN
+	//@UE BEGIN
 #if WITH_NAVMESH_SEGMENT_LINKS
 	if (poly->getType() == DT_POLYTYPE_OFFMESH_SEGMENT)
 	{
@@ -2380,7 +2380,7 @@ dtStatus dtNavMesh::getOffMeshConnectionPolyEndPoints(dtPolyRef prevRef, dtPolyR
 	}
 	else
 #endif // WITH_NAVMESH_SEGMENT_LINKS
-	//@UE4 END
+	//@UE END
 	{
 		dtVcopy(startPos, &tile->verts[poly->verts[idx0]*3]);
 		dtVcopy(endPos, &tile->verts[poly->verts[idx1]*3]);
@@ -2414,7 +2414,7 @@ const dtOffMeshConnection* dtNavMesh::getOffMeshConnectionByRef(dtPolyRef ref) c
 	return &tile->offMeshCons[idx];
 }
 
-//@UE4 BEGIN
+//@UE BEGIN
 #if WITH_NAVMESH_SEGMENT_LINKS
 const dtOffMeshSegmentConnection* dtNavMesh::getOffMeshSegmentConnectionByRef(dtPolyRef ref) const
 {
@@ -2464,7 +2464,7 @@ void dtNavMesh::updateOffMeshSegmentConnectionByUserId(unsigned int userId, unsi
 	}
 }
 #endif // WITH_NAVMESH_SEGMENT_LINKS
-//@UE4 END
+//@UE END
 
 void dtNavMesh::updateOffMeshConnectionByUserId(unsigned int userId, unsigned char newArea, unsigned short newFlags)
 {
@@ -2552,7 +2552,7 @@ dtStatus dtNavMesh::getPolyArea(dtPolyRef ref, unsigned char* resultArea) const
 	return DT_SUCCESS;
 }
 
-//@UE4 BEGIN 
+//@UE BEGIN 
 void dtNavMesh::applyWorldOffset(const float* offset)
 {
 	//Shift navmesh origin
@@ -2614,4 +2614,4 @@ void dtNavMesh::applyAreaCostOrder(unsigned char* costOrder)
 {
 	memcpy(m_areaCostOrder, costOrder, sizeof(m_areaCostOrder));
 }
-//@UE4 END
+//@UE END
diff --git a/Engine/Source/Runtime/Navmesh/Private/Detour/DetourNavMeshBuilder.cpp b/Engine/Source/Runtime/Navmesh/Private/Detour/DetourNavMeshBuilder.cpp
--- a/Engine/Source/Runtime/Navmesh/Private/Detour/DetourNavMeshBuilder.cpp
+++ b/Engine/Source/Runtime/Navmesh/Private/Detour/DetourNavMeshBuilder.cpp
@@ -386,7 +386,7 @@ bool dtCreateNavMeshData(dtNavMeshCreateParams* params, unsigned char** outData,
 	}
 	
 	// Off-mesh connections are stored as polygons, adjust values.
-	//@UE4 BEGIN
+	//@UE BEGIN
 #if WITH_NAVMESH_SEGMENT_LINKS
 	const int firstSegVert = params->vertCount + storedOffMeshConCount*2;
 	const int firstSegPoly = params->polyCount + storedOffMeshConCount;
@@ -396,7 +396,7 @@ bool dtCreateNavMeshData(dtNavMeshCreateParams* params, unsigned char** outData,
 	const int totPolyCount = params->polyCount + storedOffMeshConCount;
 	const int totVertCount = params->vertCount + storedOffMeshConCount * 2;
 #endif // WITH_NAVMESH_SEGMENT_LINKS
-	//@UE4 END
+	//@UE END
 	
 	// Find portal edges which are at tile borders.
 	int edgeCount = 0;
@@ -418,9 +418,9 @@ bool dtCreateNavMeshData(dtNavMeshCreateParams* params, unsigned char** outData,
 		}
 	}
 
-//@UE4 BEGIN: offmesh links will be added in dynamic array
+//@UE BEGIN: offmesh links will be added in dynamic array
 	const int maxLinkCount = edgeCount + portalCount*2;
-//@UE4 END
+//@UE END
 	
 	// Find unique detail vertices.
 	int uniqueDetailVertCount = 0;
@@ -472,7 +472,7 @@ bool dtCreateNavMeshData(dtNavMeshCreateParams* params, unsigned char** outData,
 	const int bvTreeSize = params->buildBvTree ? dtAlign4(sizeof(dtBVNode)*params->polyCount*2) : 0;
 	const int offMeshConsSize = dtAlign4(sizeof(dtOffMeshConnection)*storedOffMeshConCount);
 
-	//@UE4 BEGIN
+	//@UE BEGIN
 #if WITH_NAVMESH_SEGMENT_LINKS
 	const int offMeshSegSize = dtAlign4(sizeof(dtOffMeshSegmentConnection)*storedOffMeshSegCount);
 #else
@@ -486,7 +486,7 @@ bool dtCreateNavMeshData(dtNavMeshCreateParams* params, unsigned char** outData,
 	const int clustersSize = 0;
 	const int polyClustersSize = 0;
 #endif // WITH_NAVMESH_CLUSTER_LINKS
-	//@UE4 END
+	//@UE END
 
 	const int dataSize = headerSize + vertsSize + polysSize + linksSize +
 						 detailMeshesSize + detailVertsSize + detailTrisSize +
@@ -512,7 +512,7 @@ bool dtCreateNavMeshData(dtNavMeshCreateParams* params, unsigned char** outData,
 	dtBVNode* navBvtree = (dtBVNode*)d; d += bvTreeSize;
 	dtOffMeshConnection* offMeshCons = (dtOffMeshConnection*)d; d += offMeshConsSize;
 
-	//@UE4 BEGIN
+	//@UE BEGIN
 #if WITH_NAVMESH_SEGMENT_LINKS
 	dtOffMeshSegmentConnection* offMeshSegs = (dtOffMeshSegmentConnection*)d; d += offMeshSegSize;
 #endif // WITH_NAVMESH_SEGMENT_LINKS
@@ -521,7 +521,7 @@ bool dtCreateNavMeshData(dtNavMeshCreateParams* params, unsigned char** outData,
 	dtCluster* clusters = (dtCluster*)d; d += clustersSize;
 	unsigned short* polyClusters = (unsigned short*)d; d += polyClustersSize;
 #endif // WITH_NAVMESH_CLUSTER_LINKS
-	//@UE4 END
+	//@UE END
 	
 	// Store header
 	header->magic = DT_NAVMESH_MAGIC;
@@ -546,7 +546,7 @@ bool dtCreateNavMeshData(dtNavMeshCreateParams* params, unsigned char** outData,
 	header->offMeshConCount = storedOffMeshConCount;
 	header->bvNodeCount = params->buildBvTree ? params->polyCount*2 : 0;
 
-	//@UE4 BEGIN
+	//@UE BEGIN
 #if WITH_NAVMESH_SEGMENT_LINKS
 	header->offMeshSegPolyBase = firstSegPoly;
 	header->offMeshSegVertBase = firstSegVert;
@@ -556,7 +556,7 @@ bool dtCreateNavMeshData(dtNavMeshCreateParams* params, unsigned char** outData,
 #if WITH_NAVMESH_CLUSTER_LINKS
 	header->clusterCount = params->clusterCount; 
 #endif // WITH_NAVMESH_CLUSTER_LINKS
-	//@UE4 END
+	//@UE END
 
 	const int offMeshVertsBase = params->vertCount;
 	const int offMeshPolyBase = params->polyCount;
@@ -647,7 +647,7 @@ bool dtCreateNavMeshData(dtNavMeshCreateParams* params, unsigned char** outData,
 				p->setType(DT_POLYTYPE_OFFMESH_POINT);
 				n++;
 			}
-			//@UE4 BEGIN
+			//@UE BEGIN
 #if WITH_NAVMESH_SEGMENT_LINKS
 			else
 			{
@@ -662,7 +662,7 @@ bool dtCreateNavMeshData(dtNavMeshCreateParams* params, unsigned char** outData,
 				}
 			}
 #endif // WITH_NAVMESH_SEGMENT_LINKS
-			//@UE4 END
+			//@UE END
 		}
 	}
 
@@ -753,7 +753,7 @@ bool dtCreateNavMeshData(dtNavMeshCreateParams* params, unsigned char** outData,
 					con->userId = offMeshCon.userID;
 				n++;
 			}
-			//@UE4 BEGIN
+			//@UE BEGIN
 #if WITH_NAVMESH_SEGMENT_LINKS
 			else
 			{
@@ -772,13 +772,13 @@ bool dtCreateNavMeshData(dtNavMeshCreateParams* params, unsigned char** outData,
 				nseg++;
 			}
 #endif // WITH_NAVMESH_SEGMENT_LINKS
-			//@UE4 END
+			//@UE END
 		}
 	}
 
 	dtFree(offMeshConClass);
 	
-	//@UE4 BEGIN
+	//@UE BEGIN
 #if WITH_NAVMESH_CLUSTER_LINKS
 	// Store clusters
 	if (params->polyClusters)
@@ -814,7 +814,7 @@ bool dtCreateNavMeshData(dtNavMeshCreateParams* params, unsigned char** outData,
 		}
 	}
 #endif // WITH_NAVMESH_CLUSTER_LINKS
-	//@UE4 END
+	//@UE END
  
 	*outData = data;
 	*outDataSize = dataSize;
@@ -863,7 +863,7 @@ bool dtNavMeshHeaderSwapEndian(unsigned char* data, const int /*dataSize*/)
 	dtSwapEndian(&header->bmax[2]);
 	dtSwapEndian(&header->bvQuantFactor);
 
-	//@UE4 BEGIN
+	//@UE BEGIN
 #if WITH_NAVMESH_SEGMENT_LINKS
 	dtSwapEndian(&header->offMeshSegConCount);
 	dtSwapEndian(&header->offMeshSegPolyBase);
@@ -873,7 +873,7 @@ bool dtNavMeshHeaderSwapEndian(unsigned char* data, const int /*dataSize*/)
 #if WITH_NAVMESH_CLUSTER_LINKS
 	dtSwapEndian(&header->clusterCount);
 #endif // WITH_NAVMESH_CLUSTER_LINKS
-	//@UE4 END
+	//@UE END
 
 	// Freelist index and pointers are updated when tile is added, no need to swap.
 
@@ -906,7 +906,7 @@ bool dtNavMeshDataSwapEndian(unsigned char* data, const int /*dataSize*/)
 	const int bvtreeSize = dtAlign4(sizeof(dtBVNode)*header->bvNodeCount);
 	const int offMeshLinksSize = dtAlign4(sizeof(dtOffMeshConnection)*header->offMeshConCount);
 
-	//@UE4 BEGIN
+	//@UE BEGIN
 #if WITH_NAVMESH_SEGMENT_LINKS
 	const int offMeshSegSize = dtAlign4(sizeof(dtOffMeshSegmentConnection)*header->offMeshSegConCount);
 #endif // WITH_NAVMESH_SEGMENT_LINKS
@@ -915,7 +915,7 @@ bool dtNavMeshDataSwapEndian(unsigned char* data, const int /*dataSize*/)
 	const int clustersSize = dtAlign4(sizeof(dtCluster)*header->clusterCount);
 	const int polyClustersSize = dtAlign4(sizeof(unsigned short)*header->offMeshBase);
 #endif // WITH_NAVMESH_CLUSTER_LINKS
-	//@UE4 END
+	//@UE END
 
 	unsigned char* d = data + headerSize;
 	float* verts = (float*)d; d += vertsSize;
@@ -927,7 +927,7 @@ bool dtNavMeshDataSwapEndian(unsigned char* data, const int /*dataSize*/)
 	dtBVNode* bvTree = (dtBVNode*)d; d += bvtreeSize;
 	dtOffMeshConnection* offMeshCons = (dtOffMeshConnection*)d; d += offMeshLinksSize;
 
-	//@UE4 BEGIN
+	//@UE BEGIN
 #if WITH_NAVMESH_SEGMENT_LINKS
 	dtOffMeshSegmentConnection* offMeshSegs = (dtOffMeshSegmentConnection*)d; d += offMeshSegSize;
 #endif // WITH_NAVMESH_SEGMENT_LINKS	
@@ -936,7 +936,7 @@ bool dtNavMeshDataSwapEndian(unsigned char* data, const int /*dataSize*/)
 	dtCluster* clusters = (dtCluster*)d; d += clustersSize;
 	unsigned short* polyClusters = (unsigned short*)d; d += polyClustersSize;
 #endif // WITH_NAVMESH_CLUSTER_LINKS
-	//@UE4 END
+	//@UE END
 
 	// Vertices
 	for (int i = 0; i < header->vertCount*3; ++i)
@@ -996,7 +996,7 @@ bool dtNavMeshDataSwapEndian(unsigned char* data, const int /*dataSize*/)
 		dtSwapEndian(&con->userId);
 	}
 
-	//@UE4 BEGIN
+	//@UE BEGIN
 #if WITH_NAVMESH_SEGMENT_LINKS
 	// Off-mesh Connections: segment type
 	for (int i = 0; i < header->offMeshSegConCount; ++i)
@@ -1021,12 +1021,12 @@ bool dtNavMeshDataSwapEndian(unsigned char* data, const int /*dataSize*/)
 		dtSwapEndian(&polyClusters[i]);
 	}
 #endif // WITH_NAVMESH_CLUSTER_LINKS
-	//@UE4 END
+	//@UE END
  
 	return true;
 }
 
-// @UE4 BEGIN
+// @UE BEGIN
 // Experimental tile transform
 bool dtTransformTileData(unsigned char* data, const int dataSize, const int offsetX, const int offsetY, const float tileWidth, const float tileHeight, const float rotationDeg)
 {
@@ -1226,4 +1226,4 @@ void dtComputeTileOffsetFromRotation(const float* position, const float* rotatio
 	deltaX = RcTilePosChangedVector[0] / tileWidth;
 	deltaY = RcTilePosChangedVector[2] / tileHeight;
 }
-// @UE4 END
+// @UE END
diff --git a/Engine/Source/Runtime/Navmesh/Private/Detour/DetourNavMeshQuery.cpp b/Engine/Source/Runtime/Navmesh/Private/Detour/DetourNavMeshQuery.cpp
--- a/Engine/Source/Runtime/Navmesh/Private/Detour/DetourNavMeshQuery.cpp
+++ b/Engine/Source/Runtime/Navmesh/Private/Detour/DetourNavMeshQuery.cpp
@@ -68,9 +68,9 @@ DEFINE_LOG_CATEGORY_STATIC(LogDebugRaycastCrash, All, All);
 /// @see dtNavMeshQuery
 
 dtQueryFilterData::dtQueryFilterData() : heuristicScale(0.999f), lowestAreaCost(1.0f), m_includeFlags(0xffff), m_excludeFlags(0), m_isBacktracking(0)
-//@UE4 BEGIN
+//@UE BEGIN
 	, m_shouldIgnoreClosedNodes(0)
-//@UE4 END
+//@UE END
 {
 	for (int i = 0; i < DT_MAX_AREAS; ++i)
 		m_areaCost[i] = 1.0f;
@@ -87,9 +87,9 @@ bool dtQueryFilterData::equals(const dtQueryFilterData* other) const
 					(m_includeFlags == other->m_includeFlags) &&
 					(m_excludeFlags == other->m_excludeFlags) &&
 					(m_isBacktracking == other->m_isBacktracking) &&
-//@UE4 BEGIN
+//@UE BEGIN
 					(m_shouldIgnoreClosedNodes == other->m_shouldIgnoreClosedNodes) &&
-//@UE4 END
+//@UE END
 					(memcmp(&m_areaCost, &other->m_areaCost, sizeof(m_areaCost)) == 0);
 
 #if WITH_FIXED_AREA_ENTERING_COST
@@ -104,9 +104,9 @@ void dtQueryFilterData::copyFrom(const dtQueryFilterData* source)
 	memcpy((void*)this, source, sizeof(dtQueryFilterData));
 }
 
-//@UE4 BEGIN
+//@UE BEGIN
 static const float DEFAULT_HEURISTIC_SCALE = 0.999f; // Search heuristic scale.
-//@UE4 END
+//@UE END
 
 dtNavMeshQuery* dtAllocNavMeshQuery()
 {
@@ -251,9 +251,9 @@ dtStatus dtNavMeshQuery::init(const dtNavMesh* nav, const int maxNodes, dtQueryS
 			}
 			m_nodePool = new (dtAlloc(sizeof(dtNodePool), DT_ALLOC_PERM)) dtNodePool(maxNodes, dtNextPow2(maxNodes / 4));
 			if (!m_nodePool
-//@UE4 BEGIN
+//@UE BEGIN
 				&& maxNodes > 0
-//@UE4 END
+//@UE END
 				)
 				return DT_FAILURE | DT_OUT_OF_MEMORY;
 		}
@@ -261,9 +261,9 @@ dtStatus dtNavMeshQuery::init(const dtNavMesh* nav, const int maxNodes, dtQueryS
 		{
 			m_nodePool->clear();
 		}
-//@UE4 BEGIN
+//@UE BEGIN
 		m_nodePool->setMaxRuntimeNodes(maxNodes);
-//@UE4 END
+//@UE END
 
 		if (!m_tinyNodePool)
 		{
@@ -495,9 +495,9 @@ dtStatus dtNavMeshQuery::findRandomPointAroundCircle(dtPolyRef startRef, const f
 			dtPolyRef neighbourRef = link.ref;
 			// Skip invalid neighbours and do not follow back to parent.
 			if (!neighbourRef || neighbourRef == parentRef
-//@UE4 BEGIN
+//@UE BEGIN
 				|| !filter->isValidLinkSide(link.side))
-//@UE4 END
+//@UE END
 				continue;
 			
 			// Expand to neighbour
@@ -604,7 +604,7 @@ dtStatus dtNavMeshQuery::findRandomPointAroundCircle(dtPolyRef startRef, const f
 	return foundPt ? DT_SUCCESS : DT_FAILURE;
 }
 
-//@UE4 BEGIN
+//@UE BEGIN
 #if WITH_NAVMESH_CLUSTER_LINKS
 dtStatus dtNavMeshQuery::findRandomPointInCluster(dtClusterRef clusterRef, float (*frand)(), dtPolyRef* randomRef, float* randomPt) const
 {
@@ -691,7 +691,7 @@ dtStatus dtNavMeshQuery::findRandomPointInCluster(dtClusterRef clusterRef, float
 	return DT_SUCCESS;
 }
 #endif // WITH_NAVMESH_CLUSTER_LINKS
-//@UE4 END
+//@UE END
 
 //////////////////////////////////////////////////////////////////////////////////////////
 
@@ -979,7 +979,7 @@ dtStatus dtNavMeshQuery::getPolyHeight(dtPolyRef ref, const float* pos, float* h
 			*height = v0[1] + (v1[1] - v0[1]) * u;
 		return DT_SUCCESS;
 	}
-	//@UE4 BEGIN
+	//@UE BEGIN
 #if WITH_NAVMESH_SEGMENT_LINKS
 	else if (poly->getType() == DT_POLYTYPE_OFFMESH_SEGMENT)
 	{
@@ -1002,7 +1002,7 @@ dtStatus dtNavMeshQuery::getPolyHeight(dtPolyRef ref, const float* pos, float* h
 		}
 	}
 #endif // WITH_NAVMESH_SEGMENT_LINKS
-	//@UE4 END
+	//@UE END
 	else
 	{
 		const unsigned int ip = (unsigned int)(poly - tile->polys);
@@ -1031,7 +1031,7 @@ dtStatus dtNavMeshQuery::getPolyHeight(dtPolyRef ref, const float* pos, float* h
 	return DT_FAILURE | DT_INVALID_PARAM;
 }
 
-//@UE4 BEGIN
+//@UE BEGIN
 #if WITH_NAVMESH_CLUSTER_LINKS
 dtStatus dtNavMeshQuery::getPolyCluster(dtPolyRef polyRef, dtClusterRef& clusterRef) const
 {
@@ -1059,7 +1059,7 @@ dtStatus dtNavMeshQuery::getPolyCluster(dtPolyRef polyRef, dtClusterRef& cluster
 	return DT_SUCCESS;
 }
 #endif // WITH_NAVMESH_CLUSTER_LINKS
-//@UE4 END
+//@UE END
 
 /// @par 
 ///
@@ -1085,7 +1085,7 @@ dtStatus dtNavMeshQuery::findNearestPoly(const float* center, const float* exten
 	if (dtStatusFailed(queryPolygons(center, extents, filter, polys, &polyCount, 128)))
 		return DT_FAILURE | DT_INVALID_PARAM;
 	
-//@UE4 BEGIN
+//@UE BEGIN
 	float referenceLocation[3];
 	dtVcopy(referenceLocation, referencePt ? referencePt : center);
 	
@@ -1099,7 +1099,7 @@ dtStatus dtNavMeshQuery::findNearestPoly(const float* center, const float* exten
 		closestPointOnPoly(ref, referenceLocation, closestPtPoly);
 		const float d = dtVdistSqr(referenceLocation, closestPtPoly);
 		const float h = dtAbs(center[1] - closestPtPoly[1]);
-//@UE4 END
+//@UE END
 		if (d < nearestDistanceSqr && h < extents[1])
 		{
 			if (nearestPt)
@@ -1115,7 +1115,7 @@ dtStatus dtNavMeshQuery::findNearestPoly(const float* center, const float* exten
 	return DT_SUCCESS;
 }
 
-//@UE4 BEGIN
+//@UE BEGIN
 dtStatus dtNavMeshQuery::findNearestPoly2D(const float* center, const float* extents,
 											const dtQueryFilter* filter,
 											dtPolyRef* nearestRef, float* nearestPt,
@@ -1185,7 +1185,7 @@ dtStatus dtNavMeshQuery::findNearestPoly2D(const float* center, const float* ext
 
 	return DT_SUCCESS;
 }
-//@UE4 END
+//@UE END
 
 /// @par 
 ///
@@ -1443,7 +1443,7 @@ dtStatus dtNavMeshQuery::queryPolygons(const float* center, const float* extents
 ///
 dtStatus dtNavMeshQuery::findPath(dtPolyRef startRef, dtPolyRef endRef,
 								  const float* startPos, const float* endPos,
-								  const float costLimit, const dtQueryFilter* filter, //@UE4
+								  const float costLimit, const dtQueryFilter* filter, //@UE
 								  dtQueryResult& result, float* totalCost) const
 {
 	dtAssert(m_nav);
@@ -1465,10 +1465,10 @@ dtStatus dtNavMeshQuery::findPath(dtPolyRef startRef, dtPolyRef endRef,
 		return DT_SUCCESS;
 	}
 	
-	//@UE4 BEGIN
+	//@UE BEGIN
 	const float H_SCALE = filter->getModifiedHeuristicScale();
 	const bool shouldIgnoreClosedNodes = filter->getShouldIgnoreClosedNodes();
-	//@UE4 END
+	//@UE END
 
 	m_nodePool->clear();
 	m_openList->clear();
@@ -1538,9 +1538,9 @@ dtStatus dtNavMeshQuery::findPath(dtPolyRef startRef, dtPolyRef endRef,
 			
 			// Skip invalid ids and do not expand back to where we came from.
 			if (!neighbourRef || neighbourRef == parentRef
-				//@UE4 BEGIN
+				//@UE BEGIN
 				|| !filter->isValidLinkSide(link.side))
-				//@UE4 END
+				//@UE END
 			{
 				UE_RECAST_ASTAR_LOG( Warning, TEXT("Filtered %lld from %lld"), neighbourRef, bestRef);
 				continue;
@@ -1565,13 +1565,13 @@ dtStatus dtNavMeshQuery::findPath(dtPolyRef startRef, dtPolyRef endRef,
 				status |= DT_OUT_OF_NODES;
 				continue;
 			}
-//@UE4 BEGIN
+//@UE BEGIN
 			else if (shouldIgnoreClosedNodes && (neighbourNode->flags & DT_NODE_CLOSED) != 0)
 			{
 				UE_RECAST_ASTAR_LOG(Warning, TEXT("Skipping closed %lld from %lld"), neighbourRef, bestRef);
 				continue;
 			}
-//@UE4 END
+//@UE END
 
 			// Try to update node position for current edge to make paths more precise
 			// Unless heuristic is not admissible (overestimates), in which case
@@ -1631,7 +1631,7 @@ dtStatus dtNavMeshQuery::findPath(dtPolyRef startRef, dtPolyRef endRef,
 				continue;
 			}
 
-			if (total > costLimit) //@UE4
+			if (total > costLimit) //@UE
 			{
 				UE_RECAST_ASTAR_LOG(Warning, TEXT("Skipping reach cost limit poly %lld from %lld cost %f total %f prev cost %f limit %f"), neighbourRef, bestRef, cost, total, neighbourNode->total, costLimit);
 				continue;
@@ -1714,7 +1714,7 @@ dtStatus dtNavMeshQuery::findPath(dtPolyRef startRef, dtPolyRef endRef,
 	return status;
 }
 
-//@UE4 BEGIN
+//@UE BEGIN
 #if WITH_NAVMESH_CLUSTER_LINKS
 dtStatus dtNavMeshQuery::testClusterPath(dtPolyRef startRef, dtPolyRef endRef) const
 {
@@ -1869,7 +1869,7 @@ dtStatus dtNavMeshQuery::testClusterPath(dtPolyRef startRef, dtPolyRef endRef) c
 	return status;
 }
 #endif // WITH_NAVMESH_CLUSTER_LINKS
-//@UE4 END
+//@UE END
 
 /// @par
 ///
@@ -1880,7 +1880,7 @@ dtStatus dtNavMeshQuery::testClusterPath(dtPolyRef startRef, dtPolyRef endRef) c
 /// path query.
 ///
 dtStatus dtNavMeshQuery::initSlicedFindPath(dtPolyRef startRef, dtPolyRef endRef,
-											const float* startPos, const float* endPos, const float costLimit, //@UE4
+											const float* startPos, const float* endPos, const float costLimit, //@UE
 											const dtQueryFilter* filter)
 {
 	dtAssert(m_nav);
@@ -1894,7 +1894,7 @@ dtStatus dtNavMeshQuery::initSlicedFindPath(dtPolyRef startRef, dtPolyRef endRef
 	m_query.endRef = endRef;
 	dtVcopy(m_query.startPos, startPos);
 	dtVcopy(m_query.endPos, endPos);
-	m_query.costLimit = costLimit; //@UE4
+	m_query.costLimit = costLimit; //@UE
 	m_query.filter = filter;
 	
 	if (!startRef || !endRef)
@@ -1910,9 +1910,9 @@ dtStatus dtNavMeshQuery::initSlicedFindPath(dtPolyRef startRef, dtPolyRef endRef
 		return DT_SUCCESS;
 	}
 	
-	//@UE4 BEGIN
+	//@UE BEGIN
 	const float H_SCALE = filter->getModifiedHeuristicScale();
-	//@UE4 END
+	//@UE END
 
 	m_nodePool->clear();
 	m_openList->clear();
@@ -1945,10 +1945,10 @@ dtStatus dtNavMeshQuery::updateSlicedFindPath(const int maxIter, int* doneIters)
 		return DT_FAILURE;
 	}
 
-	//@UE4 BEGIN
+	//@UE BEGIN
 	const float H_SCALE = m_query.filter->getModifiedHeuristicScale();
 	const bool shouldIgnoreClosedNodes = m_query.filter->getShouldIgnoreClosedNodes();
-	//@UE4 END
+	//@UE END
 			
 	int iter = 0;
 	while (iter < maxIter && !m_openList->empty())
@@ -2013,9 +2013,9 @@ dtStatus dtNavMeshQuery::updateSlicedFindPath(const int maxIter, int* doneIters)
 			
 			// Skip invalid ids and do not expand back to where we came from.
 			if (!neighbourRef || neighbourRef == parentRef
-//@UE4 BEGIN
+//@UE BEGIN
 				|| !m_query.filter->isValidLinkSide(link.side))
-//@UE4 END
+//@UE END
 				continue;
 			
 			// Get neighbour poly and tile.
@@ -2033,12 +2033,12 @@ dtStatus dtNavMeshQuery::updateSlicedFindPath(const int maxIter, int* doneIters)
 				m_query.status |= DT_OUT_OF_NODES;
 				continue;
 			}
-//@UE4 BEGIN
+//@UE BEGIN
 			else if (shouldIgnoreClosedNodes && (neighbourNode->flags & DT_NODE_CLOSED) != 0)
 			{
 				continue;
 			}
-//@UE4 END
+//@UE END
 			
 			// Always calculate correct position on neighbor's edge,
 			// skipping to wrong edge may greatly change path cost
@@ -2051,7 +2051,7 @@ dtStatus dtNavMeshQuery::updateSlicedFindPath(const int maxIter, int* doneIters)
 			// Calculate cost and heuristic.
 			float cost = 0;
 			float heuristic = 0;
-			float curCost = 0; //@UE4
+			float curCost = 0; //@UE
 			
 			// Special case for last node.
 			if (neighbourRef != m_query.endRef)
@@ -2091,10 +2091,10 @@ dtStatus dtNavMeshQuery::updateSlicedFindPath(const int maxIter, int* doneIters)
 				continue;
 
 			// Cost of current link is DT_UNWALKABLE_POLY_COST, skip.
-			if (curCost == DT_UNWALKABLE_POLY_COST) //@UE4
+			if (curCost == DT_UNWALKABLE_POLY_COST) //@UE
 				continue;
 
-			if (total > m_query.costLimit) //@UE4
+			if (total > m_query.costLimit) //@UE
 				continue;
 			
 			// Add or update the node.
@@ -2474,7 +2474,7 @@ dtStatus dtNavMeshQuery::findStraightPath(const float* startPos, const float* en
 				fromType = toType = DT_POLYTYPE_GROUND;
 			}
 			
-			//@UE4 BEGIN
+			//@UE BEGIN
 #if WITH_NAVMESH_SEGMENT_LINKS
 			// Lock moving through segment off-mesh connections
 			if (fromType == DT_POLYTYPE_OFFMESH_SEGMENT)
@@ -2512,7 +2512,7 @@ dtStatus dtNavMeshQuery::findStraightPath(const float* startPos, const float* en
 				dtVcopy(right, tmp);
 			}
 #endif // WITH_NAVMESH_SEGMENT_LINKS
-//@UE4 END
+//@UE END
 
 			// Right vertex.
 			if (dtTriArea2D(portalApex, portalRight, right) <= 0.0f)
@@ -2554,12 +2554,12 @@ dtStatus dtNavMeshQuery::findStraightPath(const float* startPos, const float* en
 					leftIndex = apexIndex;
 					rightIndex = apexIndex;
 
-					//@UE4 BEGIN
+					//@UE BEGIN
 #if WITH_NAVMESH_SEGMENT_LINKS
 					if (toType == DT_POLYTYPE_OFFMESH_SEGMENT)
 						dtDistancePtSegSqr2D(portalApex, left, right, segt);
 #endif // WITH_NAVMESH_SEGMENT_LINKS
-					//@UE4 END
+					//@UE END
 
 					// Restart
 					i = apexIndex;
@@ -2608,12 +2608,12 @@ dtStatus dtNavMeshQuery::findStraightPath(const float* startPos, const float* en
 					leftIndex = apexIndex;
 					rightIndex = apexIndex;	
 
-					//@UE4 BEGIN
+					//@UE BEGIN
 #if WITH_NAVMESH_SEGMENT_LINKS
 					if (toType == DT_POLYTYPE_OFFMESH_SEGMENT)
 						dtDistancePtSegSqr2D(portalApex, left, right, segt);
 #endif // WITH_NAVMESH_SEGMENT_LINKS
-					//@UE4 END
+					//@UE END
 
 					// Restart
 					i = apexIndex;
@@ -2622,7 +2622,7 @@ dtStatus dtNavMeshQuery::findStraightPath(const float* startPos, const float* en
 				}
 			}
 
-			//@UE4 BEGIN
+			//@UE BEGIN
 #if WITH_NAVMESH_SEGMENT_LINKS
 			// Handle entering off-mesh segments
 			if (toType == DT_POLYTYPE_OFFMESH_SEGMENT)
@@ -2640,7 +2640,7 @@ dtStatus dtNavMeshQuery::findStraightPath(const float* startPos, const float* en
 				rightIndex = i;
 			}
 #endif // WITH_NAVMESH_SEGMENT_LINKS
-			//@UE4 END
+			//@UE END
 		}
 
 		// Append portals along the current straight path segment.
@@ -2945,7 +2945,7 @@ dtStatus dtNavMeshQuery::getPortalPoints(dtPolyRef from, const dtPoly* fromPoly,
 		}
 		return DT_FAILURE | DT_INVALID_PARAM;
 	}
-	//@UE4 BEGIN
+	//@UE BEGIN
 #if WITH_NAVMESH_SEGMENT_LINKS
 	else if (fromPoly->getType() == DT_POLYTYPE_OFFMESH_SEGMENT)
 	{
@@ -2967,7 +2967,7 @@ dtStatus dtNavMeshQuery::getPortalPoints(dtPolyRef from, const dtPoly* fromPoly,
 
 		return DT_FAILURE | DT_INVALID_PARAM;
 	}
-	//@UE4 END
+	//@UE END
 #endif // WITH_NAVMESH_SEGMENT_LINKS
 
 	if (toPoly->getType() == DT_POLYTYPE_OFFMESH_POINT)
@@ -2988,7 +2988,7 @@ dtStatus dtNavMeshQuery::getPortalPoints(dtPolyRef from, const dtPoly* fromPoly,
 		}
 		return DT_FAILURE | DT_INVALID_PARAM;
 	}
-	//@UE4 BEGIN
+	//@UE BEGIN
 #if WITH_NAVMESH_SEGMENT_LINKS
 	else if (toPoly->getType() == DT_POLYTYPE_OFFMESH_SEGMENT)
 	{
@@ -3009,7 +3009,7 @@ dtStatus dtNavMeshQuery::getPortalPoints(dtPolyRef from, const dtPoly* fromPoly,
 		return DT_FAILURE | DT_INVALID_PARAM;
 	}
 #endif // WITH_NAVMESH_SEGMENT_LINKS
-	//@UE4 END
+	//@UE END
 
 	// Find portal vertices.
 	const int v0 = fromPoly->verts[link->edge];
@@ -3019,9 +3019,9 @@ dtStatus dtNavMeshQuery::getPortalPoints(dtPolyRef from, const dtPoly* fromPoly,
 	
 	// If the link is at tile boundary, dtClamp the vertices to
 	// the link width.
-//@UE4 BEGIN
+//@UE BEGIN
 	if ((link->side & DT_CONNECTION_INTERNAL) == 0)
-//@UE4 END
+//@UE END
 	{
 		// Unpack portal limits.
 		if (link->bmin != 0 || link->bmax != 255)
@@ -3124,7 +3124,7 @@ dtStatus dtNavMeshQuery::raycast(dtPolyRef startRef, const float* startPos, cons
 	hitNormal[1] = 0;
 	hitNormal[2] = 0;
 
-	// [UE4]: iteration limit, use the same value as findPath
+	// [UE]: iteration limit, use the same value as findPath
 	const int loopLimit = (m_nodePool->getMaxRuntimeNodes() + 1) * 4;
 	int loopCounter = 0;
 
@@ -3232,9 +3232,9 @@ dtStatus dtNavMeshQuery::raycast(dtPolyRef startRef, const float* startPos, cons
 				continue;
 			
 			// If the link is internal, just return the ref.
-//@UE4 BEGIN
+//@UE BEGIN
 			if (link.side & DT_CONNECTION_INTERNAL)
-//@UE4 END
+//@UE END
 			{
 				nextRef = link.ref;
 				break;
@@ -3256,13 +3256,13 @@ dtStatus dtNavMeshQuery::raycast(dtPolyRef startRef, const float* startPos, cons
 			const float* left = &tile->verts[v0*3];
 			const float* right = &tile->verts[v1*3];
 			
-//@UE4 BEGIN
+//@UE BEGIN
 			// strip off additional flags
 			const unsigned char side = link.side & DT_LINK_FLAG_SIDE_MASK;
 
 			// Check that the intersection lies inside the link portal.
 			if (side == 0 || side == 4)
-//@UE4 END
+//@UE END
 			{
 				// Calculate link size.
 				const float s = 1.0f/255.0f;
@@ -3278,9 +3278,9 @@ dtStatus dtNavMeshQuery::raycast(dtPolyRef startRef, const float* startPos, cons
 					break;
 				}
 			}
-//@UE4 BEGIN
+//@UE BEGIN
 			else if (side == 2 || side == 6)
-//@UE4 END
+//@UE END
 			{
 				// Calculate link size.
 				const float s = 1.0f/255.0f;
@@ -3367,10 +3367,10 @@ dtStatus dtNavMeshQuery::findPolysAroundCircle(dtPolyRef startRef, const float*
 	dtAssert(m_nodePool);
 	dtAssert(m_openList);
 
-//@UE4 BEGIN
+//@UE BEGIN
 	if (!resultCount)
 		return DT_FAILURE | DT_INVALID_PARAM;
-//@UE4 END
+//@UE END
 
 	*resultCount = 0;
 	
@@ -3441,9 +3441,9 @@ dtStatus dtNavMeshQuery::findPolysAroundCircle(dtPolyRef startRef, const float*
 			dtPolyRef neighbourRef = link.ref;
 			// Skip invalid neighbours and do not follow back to parent.
 			if (!neighbourRef || neighbourRef == parentRef
-//@UE4 BEGIN
+//@UE BEGIN
 				|| !filter->isValidLinkSide(link.side))
-//@UE4 END
+//@UE END
 				continue;
 			
 			// Expand to neighbour
@@ -3625,9 +3625,9 @@ dtStatus dtNavMeshQuery::findPolysAroundShape(dtPolyRef startRef, const float* v
 			dtPolyRef neighbourRef = link.ref;
 			// Skip invalid neighbours and do not follow back to parent.
 			if (!neighbourRef || neighbourRef == parentRef
-//@UE4 BEGIN
+//@UE BEGIN
 				|| !filter->isValidLinkSide(link.side))
-//@UE4 END
+//@UE END
 				continue;
 			
 			// Expand to neighbour
@@ -3708,7 +3708,7 @@ dtStatus dtNavMeshQuery::findPolysAroundShape(dtPolyRef startRef, const float* v
 	return status;
 }
 
-//@UE4 BEGIN
+//@UE BEGIN
 //	based on dtNavMeshQuery::findPolysAroundCircle. Refer to its description for more details.
 dtStatus dtNavMeshQuery::findPolysInPathDistance(dtPolyRef startRef, const float* centerPos, const float pathDistance,
 												const dtQueryFilter* filter, dtPolyRef* resultRef,
@@ -3783,9 +3783,9 @@ dtStatus dtNavMeshQuery::findPolysInPathDistance(dtPolyRef startRef, const float
 			dtPolyRef neighbourRef = link.ref;
 			// Skip invalid neighbours and do not follow back to parent.
 			if (!neighbourRef || neighbourRef == parentRef
-//@UE4 BEGIN
+//@UE BEGIN
 				|| !filter->isValidLinkSide(link.side))
-//@UE4 END
+//@UE END
 				continue;
 
 			// Expand to neighbour
@@ -3873,7 +3873,7 @@ static bool containsPolyRef(const dtPolyRef testRef, const dtPolyRef* path, cons
 
 	return false;
 }
-//@UE4 END
+//@UE END
 
 /// @par
 ///
diff --git a/Engine/Source/Runtime/Navmesh/Private/Detour/DetourNode.cpp b/Engine/Source/Runtime/Navmesh/Private/Detour/DetourNode.cpp
--- a/Engine/Source/Runtime/Navmesh/Private/Detour/DetourNode.cpp
+++ b/Engine/Source/Runtime/Navmesh/Private/Detour/DetourNode.cpp
@@ -58,9 +58,9 @@ dtNodePool::dtNodePool(int maxNodes, int hashSize) :
 	m_next(0),
 	m_maxNodes(maxNodes),
 	m_hashSize(hashSize),
-	//@UE4 BEGIN
+	//@UE BEGIN
 	m_maxRuntimeNodes(maxNodes),
-	//@UE4 END
+	//@UE END
 	m_nodeCount(0)
 {
 	dtAssert(dtNextPow2(m_hashSize) == (unsigned int)m_hashSize);
@@ -116,9 +116,9 @@ dtNode* dtNodePool::getNode(dtPolyRef id)
 		i = m_next[i];
 	}
 
-	//@UE4 BEGIN
+	//@UE BEGIN
 	if (m_nodeCount >= getMaxRuntimeNodes())
-	//@UE4 END
+	//@UE END
 		return 0;
 	
 	i = (dtNodeIndex)m_nodeCount;
diff --git a/Engine/Source/Runtime/Navmesh/Private/DetourCrowd/DetourCrowd.cpp b/Engine/Source/Runtime/Navmesh/Private/DetourCrowd/DetourCrowd.cpp
--- a/Engine/Source/Runtime/Navmesh/Private/DetourCrowd/DetourCrowd.cpp
+++ b/Engine/Source/Runtime/Navmesh/Private/DetourCrowd/DetourCrowd.cpp
@@ -208,7 +208,7 @@ static int getNeighbours(const float* pos, const float height, const float range
 		if (distSqr > dtSqr(range))
 			continue;
 
-		// [UE4] add only when avoidance group allows it
+		// [UE] add only when avoidance group allows it
 		const bool bDontAvoid = (skip->params.groupsToIgnore & ag->params.avoidanceGroup) || !(skip->params.groupsToAvoid & ag->params.avoidanceGroup);
 		if (bDontAvoid)
 			continue;
@@ -401,7 +401,7 @@ bool dtCrowd::init(const int maxAgents, const float maxAgentRadius, dtNavMesh* n
 	if (!m_grid->init(m_maxAgents*4, maxAgentRadius*3))
 		return false;
 
-	// [UE4] moved avoidance query init to separate function
+	// [UE] moved avoidance query init to separate function
 
 	// Allocate temp buffer for merging paths.
 	m_maxPathResult = 256;
@@ -528,7 +528,7 @@ void dtCrowd::updateAgentParameters(const int idx, const dtCrowdAgentParams& par
 	m_agents[idx].params = params;
 }
 
-// [UE4] multiple filter support
+// [UE] multiple filter support
 bool dtCrowd::updateAgentFilter(const int idx, const dtQueryFilter* filter)
 {
 	if (idx < 0 || idx > m_maxAgents)
@@ -590,7 +590,7 @@ int dtCrowd::addAgent(const float* pos, const dtCrowdAgentParams& params, const
 	
 	dtCrowdAgent* ag = &m_agents[idx];
 	
-	// [UE4] multiple filter support
+	// [UE] multiple filter support
 	const bool bStoredFilter = updateAgentFilter(idx, filter);
 	if (!bStoredFilter)
 	{
@@ -854,8 +854,8 @@ void dtCrowd::updateMoveRequest(const float /*dt*/)
 			// Quick seach towards the goal.
 			static const int MAX_ITER = 20;
 			m_navquery->updateLinkFilter(ag->params.linkFilter.Get());
-			const float costLimit = FLT_MAX; //@UE4
-			m_navquery->initSlicedFindPath(path[0], ag->targetRef, ag->npos, ag->targetPos, costLimit, &m_filters[ag->params.filter]); //@UE4
+			const float costLimit = FLT_MAX; //@UE
+			m_navquery->initSlicedFindPath(path[0], ag->targetRef, ag->npos, ag->targetPos, costLimit, &m_filters[ag->params.filter]); //@UE
 			m_navquery->updateSlicedFindPath(MAX_ITER, 0);
 			dtStatus status = 0;
 			if (ag->targetReplan) // && npath > 10)
@@ -921,9 +921,9 @@ void dtCrowd::updateMoveRequest(const float /*dt*/)
 	for (int i = 0; i < nqueue; ++i)
 	{
 		dtCrowdAgent* ag = queue[i];
-		const float costLimit = FLT_MAX; //@UE4
+		const float costLimit = FLT_MAX; //@UE
 		ag->targetPathqRef = m_pathq.request(ag->corridor.getLastPoly(), ag->targetRef,
-			ag->corridor.getTarget(), ag->targetPos, costLimit, &m_filters[ag->params.filter], ag->params.linkFilter); //@UE4
+			ag->corridor.getTarget(), ag->targetPos, costLimit, &m_filters[ag->params.filter], ag->params.linkFilter); //@UE
 		if (ag->targetPathqRef != DT_PATHQ_INVALID)
 			ag->targetState = DT_CROWDAGENT_TARGET_WAITING_FOR_PATH;
 	}
@@ -1096,7 +1096,7 @@ void dtCrowd::checkPathValidity(dtCrowdAgent** agents, const int nagents, const
 		dtCrowdAgent* ag = agents[i];
 		bool replan = false;
 
-		// [UE4] try to place agent back on navmesh
+		// [UE] try to place agent back on navmesh
 		if (ag->state == DT_CROWDAGENT_STATE_INVALID)
 		{
 			ag->targetReplanTime += dt;
@@ -1150,7 +1150,7 @@ void dtCrowd::checkPathValidity(dtCrowdAgent** agents, const int nagents, const
 				ag->boundary.reset();
 				ag->state = DT_CROWDAGENT_STATE_INVALID;
 
-				// [UE4] reset replan timer, it will be used to place agent back on navmesh
+				// [UE] reset replan timer, it will be used to place agent back on navmesh
 				ag->targetReplanTime = 0;
 				continue;
 			}
@@ -1284,7 +1284,7 @@ void dtCrowd::updateStepProximityData(const float dt, dtCrowdAgentDebugInfo* deb
 		{
 			const bool bIgnoreEdgesNearLastCorner = ag->ncorners && (ag->cornerFlags[ag->ncorners - 1] & (DT_STRAIGHTPATH_OFFMESH_CONNECTION | DT_STRAIGHTPATH_END));
 
-			// UE4: move dir for segment scoring
+			// UE: move dir for segment scoring
 			float moveDir[3] = { 0.0f };
 			if (ag->ncorners)
 			{
@@ -1324,7 +1324,7 @@ void dtCrowd::updateStepNextMovePoint(const float dt, dtCrowdAgentDebugInfo* deb
 		if (ag->targetState == DT_CROWDAGENT_TARGET_NONE || ag->targetState == DT_CROWDAGENT_TARGET_VELOCITY)
 			continue;
 
-		// UE4: corridor.earlyReach support
+		// UE: corridor.earlyReach support
 		const bool bAllowCuttingCorners = (ag->boundary.getSegmentCount() == 0);
 
 		// Find corners for steering
@@ -1400,7 +1400,7 @@ void dtCrowd::updateStepNextMovePoint(const float dt, dtCrowdAgentDebugInfo* deb
 			// Adjust the path over the off-mesh connection.
 			dtPolyRef refs[2];
 
-			// UE4: m_keepOffmeshConnections support
+			// UE: m_keepOffmeshConnections support
 			const bool bCanStart = m_keepOffmeshConnections ?
 				ag->corridor.canMoveOverOffmeshConnection(ag->cornerPolys[ag->ncorners - 1], refs, ag->npos, anim->startPos, anim->endPos, m_navquery) :
 				ag->corridor.moveOverOffmeshConnection(ag->cornerPolys[ag->ncorners - 1], refs, ag->npos, anim->startPos, anim->endPos, m_navquery);
@@ -1496,7 +1496,7 @@ void dtCrowd::updateStepSteering(const float dt, dtCrowdAgentDebugInfo*)
 				float sepDot = dtVdot(diff, dvel);
 				if (sepDot < m_separationDirFilter)
 				{
-					// [UE4]: clamp to right/left vector, depending on which side nei is
+					// [UE]: clamp to right/left vector, depending on which side nei is
 					float testDir[3] = { 0, 0, 0 };
 					dtVcross(testDir, dvel, diff);
 					const bool bRightSide = (testDir[1] > 0);
@@ -1704,7 +1704,7 @@ void dtCrowd::updateStepOffMeshAnim(const float dt, dtCrowdAgentDebugInfo*)
 			// Prepare agent for walking.
 			ag->state = DT_CROWDAGENT_STATE_WALKING;
 
-			// UE4: m_keepOffmeshConnections support
+			// UE: m_keepOffmeshConnections support
 			if (m_keepOffmeshConnections)
 			{
 				ag->corridor.pruneOffmeshConenction(anim->polyRef);
@@ -1738,7 +1738,7 @@ void dtCrowd::updateStepOffMeshVelocity(const float dt, dtCrowdAgentDebugInfo*)
 	float DirLink[3];
 	float DirToEnd[3];
 
-	// UE4 version of offmesh anims, updates velocity and checks distance instead of fixed time
+	// UE version of offmesh anims, updates velocity and checks distance instead of fixed time
 	for (int i = 0; i < m_numActiveAgents; ++i)
 	{
 		dtCrowdAgent* ag = m_activeAgents[i];
@@ -1766,7 +1766,7 @@ void dtCrowd::updateStepOffMeshVelocity(const float dt, dtCrowdAgentDebugInfo*)
 			// Prepare agent for walking.
 			ag->state = DT_CROWDAGENT_STATE_WALKING;
 
-			// UE4: m_keepOffmeshConnections support
+			// UE: m_keepOffmeshConnections support
 			if (m_keepOffmeshConnections)
 			{
 				ag->corridor.pruneOffmeshConenction(anim->polyRef);
diff --git a/Engine/Source/Runtime/Navmesh/Private/DetourCrowd/DetourLocalBoundary.cpp b/Engine/Source/Runtime/Navmesh/Private/DetourCrowd/DetourLocalBoundary.cpp
--- a/Engine/Source/Runtime/Navmesh/Private/DetourCrowd/DetourLocalBoundary.cpp
+++ b/Engine/Source/Runtime/Navmesh/Private/DetourCrowd/DetourLocalBoundary.cpp
@@ -128,7 +128,7 @@ void dtLocalBoundary::update(dtPolyRef ref, const float* pos, const float collis
 	navquery->findLocalNeighbourhood(ref, pos, collisionQueryRange,
 									 filter, m_polys, 0, &m_npolys, MAX_LOCAL_POLYS);
 	
-	// [UE4] include direction to segment in score
+	// [UE] include direction to segment in score
 	float closestPt[3] = { 0.0f };
 	float dirToSeg[3] = { 0.0f };
 
@@ -148,17 +148,17 @@ void dtLocalBoundary::update(dtPolyRef ref, const float* pos, const float collis
 			if (distSqr > dtSqr(collisionQueryRange))
 				continue;
 
-			// [UE4] handle segments too far in Y (height diff)
+			// [UE] handle segments too far in Y (height diff)
 			const bool bOutsideHeightLimit = LocalBoundaryHelpers::IsOutsideHeightLimit(pos, s, s + 3, tseg);
 			if (bOutsideHeightLimit)
 			{
 				continue;
 			}
 
-			// [UE4] handle segments too close to requested position
+			// [UE] handle segments too close to requested position
 			const int segFlags = bIgnoreAtEnd ? LocalBoundaryHelpers::GetSegmentFlags(endPos, s, s + 3, collisionQueryRange) : 0;
 
-			// [UE4] include direction to segment in score
+			// [UE] include direction to segment in score
 			dtVlerp(closestPt, s, s + 3, tseg);
 			dtVsub(dirToSeg, closestPt, pos);
 			dtVnormalize(dirToSeg);
diff --git a/Engine/Source/Runtime/Navmesh/Private/DetourCrowd/DetourObstacleAvoidance.cpp b/Engine/Source/Runtime/Navmesh/Private/DetourCrowd/DetourObstacleAvoidance.cpp
--- a/Engine/Source/Runtime/Navmesh/Private/DetourCrowd/DetourObstacleAvoidance.cpp
+++ b/Engine/Source/Runtime/Navmesh/Private/DetourCrowd/DetourObstacleAvoidance.cpp
@@ -269,7 +269,7 @@ void dtObstacleAvoidanceQuery::addCircle(const float* pos, const float rad,
 
 void dtObstacleAvoidanceQuery::addSegment(const float* p, const float* q, int flags)
 {
-	// [UE4] fixed condition below. Used to be strict > comparison
+	// [UE] fixed condition below. Used to be strict > comparison
 	if (m_nsegments >= m_maxSegments)
 		return;
 	
@@ -393,7 +393,7 @@ float dtObstacleAvoidanceQuery::processSample(const float* vcand, const float cs
 			}
 		}
 		
-		// UE4: when sample is too close to segment (navmesh wall) - disable it completely
+		// UE: when sample is too close to segment (navmesh wall) - disable it completely
 		if (htmin < TooCloseToSegmentDistPct)
 		{
 			return -1.0f;
diff --git a/Engine/Source/Runtime/Navmesh/Private/DetourCrowd/DetourPathCorridor.cpp b/Engine/Source/Runtime/Navmesh/Private/DetourCrowd/DetourPathCorridor.cpp
--- a/Engine/Source/Runtime/Navmesh/Private/DetourCrowd/DetourPathCorridor.cpp
+++ b/Engine/Source/Runtime/Navmesh/Private/DetourCrowd/DetourPathCorridor.cpp
@@ -292,8 +292,8 @@ int dtPathCorridor::findCorners(float* cornerVerts, unsigned char* cornerFlags,
 		}
 	}
 
-	// UE4 BEGIN
-	// [UE4] fixed corner for early reach detection
+	// UE BEGIN
+	// [UE] fixed corner for early reach detection
 	if (m_hasNextExpectedCorner)
 	{
 		int foundIdx = -1;
@@ -343,7 +343,7 @@ int dtPathCorridor::findCorners(float* cornerVerts, unsigned char* cornerFlags,
 		}
 	}
 
-	// [UE4] Offset path points from corners
+	// [UE] Offset path points from corners
 	const dtNavMesh* nav = navquery->getAttachedNavMesh();
 	float v1[3], v2[3], dir[3];
 
@@ -393,7 +393,7 @@ int dtPathCorridor::findCorners(float* cornerVerts, unsigned char* cornerFlags,
 		}
 	}
 
-	// [UE4] Dynamic acceptance radius, depending on angle between first two segments of path
+	// [UE] Dynamic acceptance radius, depending on angle between first two segments of path
 	if (m_enableEarlyReach)
 	{
 		const bool bSame = dtVequal(m_prevMovePoint, cornerVerts);
@@ -440,7 +440,7 @@ int dtPathCorridor::findCorners(float* cornerVerts, unsigned char* cornerFlags,
 			}
 		}
 	}
-	// UE4 END
+	// UE END
 
 	return ncorners;
 }
@@ -478,7 +478,7 @@ bool dtPathCorridor::optimizePathVisibility(const float* next, const float pathO
 		return true;
 	
 	// Overshoot a little. This helps to optimize open fields in tiled meshes.
-	// UE4: changes to ray adjustment - make sure it's not going further than newDist
+	// UE: changes to ray adjustment - make sure it's not going further than newDist
 	float newDist = dtMin(dist+0.01f, pathOptimizationRange);
 	
 	// Adjust ray length.
@@ -524,8 +524,8 @@ bool dtPathCorridor::optimizePathTopology(dtNavMeshQuery* navquery, const dtQuer
 	
 	dtPolyRef res[MAX_RES];
 	int nres = 0;
-	const float costLimit = FLT_MAX; //@UE4
-	navquery->initSlicedFindPath(m_path[0], m_path[m_npath-1], m_pos, m_target, costLimit, filter); //@UE4
+	const float costLimit = FLT_MAX; //@UE
+	navquery->initSlicedFindPath(m_path[0], m_path[m_npath-1], m_pos, m_target, costLimit, filter); //@UE
 	navquery->updateSlicedFindPath(MAX_ITER, 0);
 	dtStatus status = navquery->finalizeSlicedFindPathPartial(m_path, m_npath, res, &nres, MAX_RES);
 	
@@ -670,7 +670,7 @@ bool dtPathCorridor::movePosition(const float* npos, dtNavMeshQuery* navquery, c
 	dtPolyRef visited[MAX_VISITED];
 	int nvisited = 0;
 
-	// [UE4: check status, it may fail due to runtime navmesh rebuild]
+	// [UE: check status, it may fail due to runtime navmesh rebuild]
 	const dtStatus status = navquery->moveAlongSurface(m_path[0], m_pos, npos, filter, result, visited, &nvisited, MAX_VISITED);
 	if (dtStatusFailed(status))
 	{
diff --git a/Engine/Source/Runtime/Navmesh/Private/DetourCrowd/DetourPathQueue.cpp b/Engine/Source/Runtime/Navmesh/Private/DetourCrowd/DetourPathQueue.cpp
--- a/Engine/Source/Runtime/Navmesh/Private/DetourCrowd/DetourPathQueue.cpp
+++ b/Engine/Source/Runtime/Navmesh/Private/DetourCrowd/DetourPathQueue.cpp
@@ -111,7 +111,7 @@ void dtPathQueue::update(const int maxIters)
 		// Handle query start.
 		if (q.status == 0)
 		{
-			q.status = m_navquery->initSlicedFindPath(q.startRef, q.endRef, q.startPos, q.endPos, q.costLimit, q.filter); //@UE4
+			q.status = m_navquery->initSlicedFindPath(q.startRef, q.endRef, q.startPos, q.endPos, q.costLimit, q.filter); //@UE
 		}		
 		// Handle query in progress.
 		if (dtStatusInProgress(q.status))
@@ -133,7 +133,7 @@ void dtPathQueue::update(const int maxIters)
 }
 
 dtPathQueueRef dtPathQueue::request(dtPolyRef startRef, dtPolyRef endRef,
-									const float* startPos, const float* endPos, const float costLimit, //@UE4
+									const float* startPos, const float* endPos, const float costLimit, //@UE
 									const dtQueryFilter* filter, TSharedPtr<dtQuerySpecialLinkFilter> linkFilter)
 {
 	// Find empty slot
@@ -159,7 +159,7 @@ dtPathQueueRef dtPathQueue::request(dtPolyRef startRef, dtPolyRef endRef,
 	q.startRef = startRef;
 	dtVcopy(q.endPos, endPos);
 	q.endRef = endRef;
-	q.costLimit = costLimit; //@UE4
+	q.costLimit = costLimit; //@UE
 	
 	q.status = 0;
 	q.npath = 0;
diff --git a/Engine/Source/Runtime/Navmesh/Private/DetourTileCache/DetourTileCache.cpp b/Engine/Source/Runtime/Navmesh/Private/DetourTileCache/DetourTileCache.cpp
--- a/Engine/Source/Runtime/Navmesh/Private/DetourTileCache/DetourTileCache.cpp
+++ b/Engine/Source/Runtime/Navmesh/Private/DetourTileCache/DetourTileCache.cpp
@@ -74,11 +74,11 @@ struct BuildContext
 		: layer(0)
 		, dfield(0)
 		, lcset(0)
-		//@UE4 BEGIN
+		//@UE BEGIN
 #if WITH_NAVMESH_CLUSTER_LINKS
 		, lclusters(0)
 #endif //WITH_NAVMESH_CLUSTER_LINKS
-		//@UE4 END
+		//@UE END
 		, lmesh(0)
 		, alloc(a)
 	{}
@@ -91,23 +91,23 @@ struct BuildContext
 		dfield = 0;
 		dtFreeTileCacheContourSet(alloc, lcset);
 		lcset = 0;
-		//@UE4 BEGIN
+		//@UE BEGIN
 #if WITH_NAVMESH_CLUSTER_LINKS
 		dtFreeTileCacheClusterSet(alloc, lclusters);
 		lclusters = 0;
 #endif //WITH_NAVMESH_CLUSTER_LINKS
-		//@UE4 END
+		//@UE END
 		dtFreeTileCachePolyMesh(alloc, lmesh);
 		lmesh = 0;
 	}
 	struct dtTileCacheLayer* layer;
 	struct dtTileCacheDistanceField* dfield;
 	struct dtTileCacheContourSet* lcset;
-	//@UE4 BEGIN
+	//@UE BEGIN
 #if WITH_NAVMESH_CLUSTER_LINKS
 	struct dtTileCacheClusterSet* lclusters;
 #endif //WITH_NAVMESH_CLUSTER_LINKS
-	//@UE4 END
+	//@UE END
 	struct dtTileCachePolyMesh* lmesh;
 	struct dtTileCacheAlloc* alloc;
 };
@@ -700,7 +700,7 @@ dtStatus dtTileCache::buildNavMeshTile(const dtCompressedTileRef ref, dtNavMesh*
 	if (!bc.lcset)
 		return status;
 
-	//@UE4 BEGIN
+	//@UE BEGIN
 #if WITH_NAVMESH_CLUSTER_LINKS
 	bc.lclusters = dtAllocTileCacheClusterSet(m_talloc);
 	if (!bc.lclusters)
@@ -712,7 +712,7 @@ dtStatus dtTileCache::buildNavMeshTile(const dtCompressedTileRef ref, dtNavMesh*
 	status = dtBuildTileCacheContours(m_talloc, *bc.layer, walkableClimbVx,
 									  m_params.maxSimplificationError, m_params.cs, m_params.ch, *bc.lcset);
 #endif //WITH_NAVMESH_CLUSTER_LINKS
-	//@UE4 END
+	//@UE END
 
 	if (dtStatusFailed(status))
 		return status;
@@ -728,13 +728,13 @@ dtStatus dtTileCache::buildNavMeshTile(const dtCompressedTileRef ref, dtNavMesh*
 	if (!bc.lmesh->npolys)
 		return DT_SUCCESS;
 
-	//@UE4 BEGIN
+	//@UE BEGIN
 #if WITH_NAVMESH_CLUSTER_LINKS
 	status = dtBuildTileCacheClusters(m_talloc, *bc.lclusters, *bc.lmesh);
 	if (dtStatusFailed(status))
 		return status;
 #endif // WITH_NAVMESH_CLUSTER_LINKS
-	//@UE4 END
+	//@UE END
 	
 	dtNavMeshCreateParams params;
 	memset(&params, 0, sizeof(params));
@@ -756,12 +756,12 @@ dtStatus dtTileCache::buildNavMeshTile(const dtCompressedTileRef ref, dtNavMesh*
 	params.buildBvTree = false;
 	dtVcopy(params.bmin, tile->header->bmin);
 	dtVcopy(params.bmax, tile->header->bmax);
-	//@UE4 BEGIN
+	//@UE BEGIN
 #if WITH_NAVMESH_CLUSTER_LINKS
 	params.polyClusters = bc.lclusters->polyMap;
 	params.clusterCount = (unsigned short)bc.lclusters->nclusters;
 #endif //WITH_NAVMESH_CLUSTER_LINKS
-	//@UE4 END
+	//@UE END
 	
 	if (m_tmproc)
 	{
diff --git a/Engine/Source/Runtime/Navmesh/Private/DetourTileCache/DetourTileCacheBuilder.cpp b/Engine/Source/Runtime/Navmesh/Private/DetourTileCache/DetourTileCacheBuilder.cpp
--- a/Engine/Source/Runtime/Navmesh/Private/DetourTileCache/DetourTileCacheBuilder.cpp
+++ b/Engine/Source/Runtime/Navmesh/Private/DetourTileCache/DetourTileCacheBuilder.cpp
@@ -47,7 +47,7 @@ void dtFreeTileCacheContourSet(dtTileCacheAlloc* alloc, dtTileCacheContourSet* c
 	alloc->free(cset);
 }
 
-//@UE4 BEGIN
+//@UE BEGIN
 #if WITH_NAVMESH_CLUSTER_LINKS
 dtTileCacheClusterSet* dtAllocTileCacheClusterSet(dtTileCacheAlloc* alloc)
 {
@@ -68,7 +68,7 @@ void dtFreeTileCacheClusterSet(dtTileCacheAlloc* alloc, dtTileCacheClusterSet* c
 	alloc->free(clusters);
 }
 #endif // WITH_NAVMESH_CLUSTER_LINKS
-//@UE4 END
+//@UE END
 
 dtTileCachePolyMesh* dtAllocTileCachePolyMesh(dtTileCacheAlloc* alloc)
 {
@@ -717,11 +717,11 @@ dtStatus dtBuildTileCacheContours(dtTileCacheAlloc* alloc, dtTileCacheLayer& lay
 	const int walkableClimb, const float maxError,
 	const float cs, const float ch,
 	dtTileCacheContourSet& lcset
-	//@UE4 BEGIN
+	//@UE BEGIN
 #if WITH_NAVMESH_CLUSTER_LINKS
 	, dtTileCacheClusterSet& clusters
 #endif //WITH_NAVMESH_CLUSTER_LINKS
-	//@UE4 END
+	//@UE END
 	)
 {
 	dtAssert(alloc);
@@ -940,7 +940,7 @@ dtStatus dtBuildTileCacheContours(dtTileCacheAlloc* alloc, dtTileCacheLayer& lay
 		}
 	}
 
-	//@UE4 BEGIN
+	//@UE BEGIN
 #if WITH_NAVMESH_CLUSTER_LINKS
 	// Build clusters
 	clusters.nregs = layer.regCount ? (layer.regCount + 1) : 0;
@@ -1034,7 +1034,7 @@ dtStatus dtBuildTileCacheContours(dtTileCacheAlloc* alloc, dtTileCacheLayer& lay
 		}
 	}
 #endif // WITH_NAVMESH_CLUSTER_LINKS
-	//@UE4 END
+	//@UE END
 
 	return DT_SUCCESS;
 }	
@@ -2001,7 +2001,7 @@ dtStatus dtBuildTileCachePolyMesh(dtTileCacheAlloc* alloc,
 	
 	mesh.nvp = MAX_VERTS_PER_POLY;
 	
-// @UE4 BEGIN: special handling of "no valid contours"
+// @UE BEGIN: special handling of "no valid contours"
 	if (maxVertices == 0)
 	{
 		// treating this as success because no issues arised
@@ -2010,7 +2010,7 @@ dtStatus dtBuildTileCachePolyMesh(dtTileCacheAlloc* alloc,
 		// by dtAllocTileCachePolyMesh
 		return DT_SUCCESS;
 	}
-// @UE4 END
+// @UE END
 
 	dtFixedArray<unsigned char> vflags(alloc, maxVertices);
 	if (!vflags)
@@ -2580,7 +2580,7 @@ dtStatus dtReplaceArea(dtTileCacheLayer& layer, const unsigned char areaId, cons
 	return DT_SUCCESS;
 }
 
-//@UE4 BEGIN
+//@UE BEGIN
 #if WITH_NAVMESH_CLUSTER_LINKS
 dtStatus dtBuildTileCacheClusters(dtTileCacheAlloc* alloc, dtTileCacheClusterSet& lclusters, dtTileCachePolyMesh& lmesh)
 {
@@ -2612,7 +2612,7 @@ dtStatus dtBuildTileCacheClusters(dtTileCacheAlloc* alloc, dtTileCacheClusterSet
 	return DT_SUCCESS;
 }
 #endif // WITH_NAVMESH_CLUSTER_LINKS
-//@UE4 END
+//@UE END
 
 dtStatus dtBuildTileCacheLayer(dtTileCacheCompressor* comp,
 							   dtTileCacheLayerHeader* header,
diff --git a/Engine/Source/Runtime/Navmesh/Private/DetourTileCache/DetourTileCacheDetail.cpp b/Engine/Source/Runtime/Navmesh/Private/DetourTileCache/DetourTileCacheDetail.cpp
--- a/Engine/Source/Runtime/Navmesh/Private/DetourTileCache/DetourTileCacheDetail.cpp
+++ b/Engine/Source/Runtime/Navmesh/Private/DetourTileCache/DetourTileCacheDetail.cpp
@@ -202,10 +202,10 @@ static unsigned short getHeight(const float fx, const float fy, const float fz,
 	unsigned short h = hp.data[ix+iz*hp.width];
 	if (h == DT_UNSET_PATCH_HEIGHT)
 	{
-		//@UE4 BEGIN
+		//@UE BEGIN
 		// setting fallback value in case proper height is not found
 		h = (unsigned short)floorf(fy/ch);
-		//@UE4 END
+		//@UE END
 
 		// Special case when data might be bad.
 		// Find nearest neighbour pixel which has valid height.
@@ -511,7 +511,7 @@ namespace TileCacheFunc
 					continue;
 				}
 				const float d = vdist2(c, &pts[u * 3]);
-				// UE4: increased tolerance of safe checks from 0.001f
+				// UE: increased tolerance of safe checks from 0.001f
 				// it was producing (rarely) overlapping edges
 				const float tol = 0.005f;
 				if (d > r*(1 + tol))
diff --git a/Engine/Source/Runtime/Navmesh/Private/DetourTileCache/DetourTileCacheRegion.cpp b/Engine/Source/Runtime/Navmesh/Private/DetourTileCache/DetourTileCacheRegion.cpp
--- a/Engine/Source/Runtime/Navmesh/Private/DetourTileCache/DetourTileCacheRegion.cpp
+++ b/Engine/Source/Runtime/Navmesh/Private/DetourTileCache/DetourTileCacheRegion.cpp
@@ -1152,7 +1152,7 @@ static dtStatus CollectRegionsMonotone(dtTileCacheAlloc* alloc, dtTileCacheLayer
 	// Allocate and init layer regions.
 	nregs = (int)regId;
 
-// @UE4 BEGIN: special handling of "no regions"
+// @UE BEGIN: special handling of "no regions"
 	if (nregs == 0)
 	{
 		regs = 0;
@@ -1160,7 +1160,7 @@ static dtStatus CollectRegionsMonotone(dtTileCacheAlloc* alloc, dtTileCacheLayer
 		// no issues occurred, everything was good. Just no regions. 
 		return DT_SUCCESS;
 	}
-// @UE4 END
+// @UE END
 
 	regs = (dtLayerMonotoneRegion*)alloc->alloc(sizeof(dtLayerMonotoneRegion) * nregs);
 	if (!regs)
@@ -1323,7 +1323,7 @@ static dtStatus CollectRegionsChunky(dtTileCacheAlloc* alloc, dtTileCacheLayer&
 	// Allocate and init layer regions.
 	nregs = (int)regId;
 
-// @UE4 BEGIN: special handling of "no regions"
+// @UE BEGIN: special handling of "no regions"
 	if (nregs == 0)
 	{
 		regs = 0;
@@ -1331,7 +1331,7 @@ static dtStatus CollectRegionsChunky(dtTileCacheAlloc* alloc, dtTileCacheLayer&
 		// no issues occurred, everything was good. Just no regions. 
 		return DT_SUCCESS;
 	}
-// @UE4 END
+// @UE END
 
 	regs = (dtLayerMonotoneRegion*)alloc->alloc(sizeof(dtLayerMonotoneRegion) * nregs);
 	if (!regs)
diff --git a/Engine/Source/Runtime/Navmesh/Private/Recast/Recast.cpp b/Engine/Source/Runtime/Navmesh/Private/Recast/Recast.cpp
--- a/Engine/Source/Runtime/Navmesh/Private/Recast/Recast.cpp
+++ b/Engine/Source/Runtime/Navmesh/Private/Recast/Recast.cpp
@@ -154,7 +154,7 @@ void rcFreeContourSet(rcContourSet* cset)
 	rcFree(cset);
 }
 
-//@UE4 BEGIN
+//@UE BEGIN
 #if WITH_NAVMESH_CLUSTER_LINKS
 rcClusterSet* rcAllocClusterSet()
 {
@@ -172,7 +172,7 @@ void rcFreeClusterSet(rcClusterSet* clusters)
 	rcFree(clusters);
 }
 #endif // WITH_NAVMESH_CLUSTER_LINKS
-//@UE4 END
+//@UE END
 
 rcPolyMesh* rcAllocPolyMesh()
 {
@@ -416,14 +416,14 @@ bool rcBuildCompactHeightfield(rcContext* ctx, const int walkableHeight, const i
 {
 	rcAssert(ctx);
 	
-// @UE4 BEGIN: early-out when no walkable spans 
+// @UE BEGIN: early-out when no walkable spans 
 	const int spanCount = rcGetHeightFieldSpanCount(ctx, hf);
 	if (spanCount == 0)
 	{
 		// no spans to speak of, bail out.
 		return false;
 	}
-// @UE4 END
+// @UE END
 	ctx->startTimer(RC_TIMER_BUILD_COMPACTHEIGHTFIELD);
 	
 	const int w = hf.width;
diff --git a/Engine/Source/Runtime/Navmesh/Private/Recast/RecastAlloc.cpp b/Engine/Source/Runtime/Navmesh/Private/Recast/RecastAlloc.cpp
--- a/Engine/Source/Runtime/Navmesh/Private/Recast/RecastAlloc.cpp
+++ b/Engine/Source/Runtime/Navmesh/Private/Recast/RecastAlloc.cpp
@@ -92,7 +92,7 @@ void rcIntArray::resize(int n)
 	m_size = n;
 }
 
-// @UE4 BEGIN: added contains()
+// @UE BEGIN: added contains()
 bool rcIntArray::contains(int v) const
 {
 	for (int i = 0; i < m_size; i++)
@@ -103,4 +103,4 @@ bool rcIntArray::contains(int v) const
 
 	return false;
 }
-// @UE4 END
+// @UE END
diff --git a/Engine/Source/Runtime/Navmesh/Private/Recast/RecastContour.cpp b/Engine/Source/Runtime/Navmesh/Private/Recast/RecastContour.cpp
--- a/Engine/Source/Runtime/Navmesh/Private/Recast/RecastContour.cpp
+++ b/Engine/Source/Runtime/Navmesh/Private/Recast/RecastContour.cpp
@@ -883,7 +883,7 @@ static int findContourFromSet(const rcContourSet& cset, unsigned short reg)
 	return -1;
 }
 
-//@UE4 BEGIN
+//@UE BEGIN
 #if WITH_NAVMESH_CLUSTER_LINKS
 bool rcBuildClusters(rcContext* ctx, rcContourSet& cset, rcClusterSet& clusters)
 {
@@ -951,4 +951,4 @@ bool rcBuildClusters(rcContext* ctx, rcContourSet& cset, rcClusterSet& clusters)
 	return true;
 }
 #endif // WITH_NAVMESH_CLUSTER_LINKS
-//@UE4 END
\ No newline at end of file
+//@UE END
\ No newline at end of file
diff --git a/Engine/Source/Runtime/Navmesh/Private/Recast/RecastFilter.cpp b/Engine/Source/Runtime/Navmesh/Private/Recast/RecastFilter.cpp
--- a/Engine/Source/Runtime/Navmesh/Private/Recast/RecastFilter.cpp
+++ b/Engine/Source/Runtime/Navmesh/Private/Recast/RecastFilter.cpp
@@ -222,7 +222,7 @@ void rcFilterWalkableLowHeightSpansSequences(rcContext* ctx, int walkableHeight,
 	int32 NumSpans;
 	memset(SpanList, 0, sizeof(SpanList));
 
-	// UE4: leave only single low span below valid one (null area doesn't count) or after leaving walkableHeight space between them
+	// UE: leave only single low span below valid one (null area doesn't count) or after leaving walkableHeight space between them
 
 	// Remove walkable flag from spans which do not have enough
 	// space above them for the agent to stand there.
diff --git a/Engine/Source/Runtime/Navmesh/Private/Recast/RecastLayers.cpp b/Engine/Source/Runtime/Navmesh/Private/Recast/RecastLayers.cpp
--- a/Engine/Source/Runtime/Navmesh/Private/Recast/RecastLayers.cpp
+++ b/Engine/Source/Runtime/Navmesh/Private/Recast/RecastLayers.cpp
@@ -63,7 +63,7 @@ inline bool overlapRange(const unsigned short amin, const unsigned short amax,
 
 static void fixLayerConnections(rcHeightfieldLayer* layer)
 {
-	// [UE4: break one directional connections, contour tracing gets stuck in infinite loop]
+	// [UE: break one directional connections, contour tracing gets stuck in infinite loop]
 	const int lw = layer->width;
 	const int lh = layer->height;
 
@@ -227,7 +227,7 @@ static bool CollectLayerRegionsMonotone(rcContext* ctx, rcCompactHeightfield& ch
 	// Allocate and init layer regions.
 	nregs = (int)regId;
 
-// @UE4 BEGIN: special handling of "no regions"
+// @UE BEGIN: special handling of "no regions"
 	if (nregs == 0)
 	{
 		regs = 0;
@@ -235,7 +235,7 @@ static bool CollectLayerRegionsMonotone(rcContext* ctx, rcCompactHeightfield& ch
 		// no issues occurred, everything was good. Just no regions. 
 		return true;
 	}
-// @UE4 END
+// @UE END
 
 	regs = (rcLayerRegionMonotone*)rcAlloc(sizeof(rcLayerRegionMonotone)*nregs, RC_ALLOC_TEMP);
 	if (!regs)
@@ -366,7 +366,7 @@ static bool CollectLayerRegionsChunky(rcContext* ctx, rcCompactHeightfield& chf,
 						if (sid == 0xffff)
 						{
 							sid = sweepId++;
-							// UE4: multiple spans per single X row may result in more sweeps than originally allocated
+							// UE: multiple spans per single X row may result in more sweeps than originally allocated
 							if (sweeps.resizeGrow(sid + 1))
 							{
 								sweeps[sid].nei = 0xffff;
@@ -443,7 +443,7 @@ static bool CollectLayerRegionsChunky(rcContext* ctx, rcCompactHeightfield& chf,
 	// Allocate and init layer regions.
 	nregs = (int)regId;
 
-// @UE4 BEGIN: special handling of "no regions"
+// @UE BEGIN: special handling of "no regions"
 	if (nregs == 0)
 	{
 		regs = 0;
@@ -451,7 +451,7 @@ static bool CollectLayerRegionsChunky(rcContext* ctx, rcCompactHeightfield& chf,
 		// no issues occurred, everything was good. Just no regions. 
 		return true;
 	}
-// @UE4 END
+// @UE END
 
 	regs = (rcLayerRegionMonotone*)rcAlloc(sizeof(rcLayerRegionMonotone)*nregs, RC_ALLOC_TEMP);
 	if (!regs)
diff --git a/Engine/Source/Runtime/Navmesh/Private/Recast/RecastMeshDetail.cpp b/Engine/Source/Runtime/Navmesh/Private/Recast/RecastMeshDetail.cpp
--- a/Engine/Source/Runtime/Navmesh/Private/Recast/RecastMeshDetail.cpp
+++ b/Engine/Source/Runtime/Navmesh/Private/Recast/RecastMeshDetail.cpp
@@ -199,17 +199,17 @@ static unsigned short getHeight(const float fx, const float fy, const float fz,
 {
 	int ix = (int)floorf(fx*ics + 0.01f);
 	int iz = (int)floorf(fz*ics + 0.01f);
-	//@UE4 BEGIN: clamp in range 0..width-1 and 0..height-1 (breaks array addressing without -1)
+	//@UE BEGIN: clamp in range 0..width-1 and 0..height-1 (breaks array addressing without -1)
 	ix = rcClamp(ix-hp.xmin, 0, hp.width-1);
 	iz = rcClamp(iz-hp.ymin, 0, hp.height-1);
-	//@UE4 END
+	//@UE END
 	unsigned short h = hp.data[ix+iz*hp.width];
 	if (h == RC_UNSET_HEIGHT)
 	{
-		//@UE4 BEGIN
+		//@UE BEGIN
 		// setting fallback value in case proper height is not found
 		h = (unsigned short)floorf(fy/ch);
-		//@UE4 END
+		//@UE END
 
 		// Special case when data might be bad.
 		// Find nearest neighbour pixel which has valid height.
diff --git a/Engine/Source/Runtime/Navmesh/Private/Recast/RecastRasterization.cpp b/Engine/Source/Runtime/Navmesh/Private/Recast/RecastRasterization.cpp
--- a/Engine/Source/Runtime/Navmesh/Private/Recast/RecastRasterization.cpp
+++ b/Engine/Source/Runtime/Navmesh/Private/Recast/RecastRasterization.cpp
@@ -122,7 +122,7 @@ static void addSpan(rcHeightfield& hf, const int x, const int y,
 		}
 		else
 		{
-			// @UE4 BEGIN
+			// @UE BEGIN
 			// Merge overlapping spans.
 
 			// For spans whose tops are really close to each other, prefer walkable areas.
@@ -143,7 +143,7 @@ static void addSpan(rcHeightfield& hf, const int x, const int y,
 				s->data.smin = cur->data.smin;
 			if (cur->data.smax > s->data.smax)
 				s->data.smax = cur->data.smax;
-			// @UE4 END
+			// @UE END
 			
 			// Remove current span.
 			rcSpan* next = cur->next;
@@ -440,15 +440,15 @@ static void rasterizeTri(const float* v0, const float* v1, const float* v2,
 						 const float* bmin, const float* bmax,
 						 const float cs, const float ics, const float ich, 
 						 const int flagMergeThr,
-						 const int rasterizationFlags, /*UE4*/
-	                     const int* rasterizationMasks /*UE4*/)
+						 const int rasterizationFlags, /*UE*/
+	                     const int* rasterizationMasks /*UE*/)
 {
 	rcEdgeHit* const hfEdgeHits = hf.EdgeHits; //this prevents a static analysis warning
 
 	const int w = hf.width;
 	const int h = hf.height;
 	const float by = bmax[1] - bmin[1];
-	const int projectTriToBottom = rasterizationFlags & RC_PROJECT_TO_BOTTOM; //UE4
+	const int projectTriToBottom = rasterizationFlags & RC_PROJECT_TO_BOTTOM; //UE
 
 	int intverts[3][2];
 
@@ -486,10 +486,10 @@ static void rasterizeTri(const float* v0, const float* v1, const float* v2,
 		// Snap the span to the heightfield height grid.
 		unsigned short triangle_ismin = (unsigned short)rcClamp((int)floorf(triangle_smin * ich), 0, RC_SPAN_MAX_HEIGHT);
 		unsigned short triangle_ismax = (unsigned short)rcClamp((int)ceilf(triangle_smax * ich), (int)triangle_ismin+1, RC_SPAN_MAX_HEIGHT);
-		const int projectSpanToBottom = rasterizationMasks != nullptr ? (projectTriToBottom & rasterizationMasks[x0+y0*w]) : projectTriToBottom;	//UE4
-		if (projectSpanToBottom) //UE4
+		const int projectSpanToBottom = rasterizationMasks != nullptr ? (projectTriToBottom & rasterizationMasks[x0+y0*w]) : projectTriToBottom;	//UE
+		if (projectSpanToBottom) //UE
 		{
-			triangle_ismin = 0; //UE4
+			triangle_ismin = 0; //UE
 		}
 
 		addSpan(hf, x0, y0, triangle_ismin, triangle_ismax, area, flagMergeThr);
@@ -648,14 +648,14 @@ static void rasterizeTri(const float* v0, const float* v1, const float* v2,
 			}
 		}
 
-		if (rasterizationMasks == nullptr) //UE4
+		if (rasterizationMasks == nullptr) //UE
 		{
 			// Snap the span to the heightfield height grid.
 			unsigned short triangle_ismin_clamp = (unsigned short)rcClamp((int)triangle_ismin, 0, RC_SPAN_MAX_HEIGHT);
 			const unsigned short triangle_ismax_clamp = (unsigned short)rcClamp((int)triangle_ismax, (int)triangle_ismin_clamp+1, RC_SPAN_MAX_HEIGHT);
-			if (projectTriToBottom) //UE4
+			if (projectTriToBottom) //UE
 			{
-				triangle_ismin_clamp = 0; //UE4
+				triangle_ismin_clamp = 0; //UE
 			}
 
 			for (int y = y0; y <= y1; y++)
@@ -674,7 +674,7 @@ static void rasterizeTri(const float* v0, const float* v1, const float* v2,
 		}
 		else
 		{
-// @UE4 BEGIN
+// @UE BEGIN
 			for (int y = y0; y <= y1; y++)
 			{
 				int xloop0 = intMax(hf.RowExt[y + 1].MinCol, x0);
@@ -684,10 +684,10 @@ static void rasterizeTri(const float* v0, const float* v1, const float* v2,
 					// Snap the span to the heightfield height grid.
 					unsigned short triangle_ismin_clamp = (unsigned short)rcClamp((int)triangle_ismin, 0, RC_SPAN_MAX_HEIGHT);
 					const unsigned short triangle_ismax_clamp = (unsigned short)rcClamp((int)triangle_ismax, (int)triangle_ismin_clamp+1, RC_SPAN_MAX_HEIGHT);
-					const int projectSpanToBottom = projectTriToBottom & rasterizationMasks[x+y*w];		//UE4
-					if (projectSpanToBottom) //UE4
+					const int projectSpanToBottom = projectTriToBottom & rasterizationMasks[x+y*w];		//UE
+					if (projectSpanToBottom) //UE
 					{
-						triangle_ismin_clamp = 0; //UE4
+						triangle_ismin_clamp = 0; //UE
 					}
 					addSpan(hf, x, y, triangle_ismin_clamp, triangle_ismax_clamp, area, flagMergeThr);
 				}
@@ -696,7 +696,7 @@ static void rasterizeTri(const float* v0, const float* v1, const float* v2,
 				hf.RowExt[y + 1].MinCol = hf.width + 2;
 				hf.RowExt[y + 1].MaxCol = -2;
 			}
-// @UE4 END
+// @UE END
 		}
 	}
 	else
@@ -865,10 +865,10 @@ static void rasterizeTri(const float* v0, const float* v1, const float* v2,
 
 				smin = intMax(smin, 0);
 				smax = intMin(intMax(smax,smin+1), RC_SPAN_MAX_HEIGHT);
-				const int projectSpanToBottom = rasterizationMasks != nullptr ? (projectTriToBottom & rasterizationMasks[x+y*w]) : projectTriToBottom; //UE4
-				if (projectSpanToBottom) //UE4
+				const int projectSpanToBottom = rasterizationMasks != nullptr ? (projectTriToBottom & rasterizationMasks[x+y*w]) : projectTriToBottom; //UE
+				if (projectSpanToBottom) //UE
 				{
-					smin = 0; //UE4
+					smin = 0; //UE
 				}
 
 	#if TEST_NEW_RASTERIZER
@@ -935,14 +935,14 @@ static void rasterizeTri(const float* v0, const float* v1, const float* v2,
 						 const float* bmin, const float* bmax,
 						 const float cs, const float ics, const float ich,
 						 const int flagMergeThr,
-						 const int rasterizationFlags,  //UE4
-						 const int* rasterizationMasks) //UE4
+						 const int rasterizationFlags,  //UE
+						 const int* rasterizationMasks) //UE
 {
 	const int w = hf.width;
 	const int h = hf.height;
 	float tmin[3], tmax[3];
 	const float by = bmax[1] - bmin[1];
-	const int projectTriToBottom = rasterizationFlags & RC_PROJECT_TO_BOTTOM; //UE4
+	const int projectTriToBottom = rasterizationFlags & RC_PROJECT_TO_BOTTOM; //UE
 	
 	// Calculate the bounding box of the triangle.
 	rcVcopy(tmin, v0);
@@ -1011,10 +1011,10 @@ static void rasterizeTri(const float* v0, const float* v1, const float* v2,
 			// Snap the span to the heightfield height grid.
 			unsigned short ismin = (unsigned short)rcClamp((int)floorf(smin * ich), 0, RC_SPAN_MAX_HEIGHT);
 			unsigned short ismax = (unsigned short)rcClamp((int)ceilf(smax * ich), (int)ismin+1, RC_SPAN_MAX_HEIGHT);
-			const int projectSpanToBottom = rasterizationMasks != nullptr ? (projectTriToBottom & rasterizationMasks[x+y*w]) : projectTriToBottom;	//UE4
-			if (projectSpanToBottom) //UE4
+			const int projectSpanToBottom = rasterizationMasks != nullptr ? (projectTriToBottom & rasterizationMasks[x+y*w]) : projectTriToBottom;	//UE
+			if (projectSpanToBottom) //UE
 			{
-				ismin = 0; //UE4
+				ismin = 0; //UE
 			}
 
 			addSpan(hf, x, y, ismin, ismax, area, flagMergeThr);
@@ -1030,7 +1030,7 @@ static void rasterizeTri(const float* v0, const float* v1, const float* v2,
 /// @see rcHeightfield
 void rcRasterizeTriangle(rcContext* ctx, const float* v0, const float* v1, const float* v2,
 						 const unsigned char area, rcHeightfield& solid,
-						 const int flagMergeThr, const int rasterizationFlags, const int* rasterizationMasks) //UE4
+						 const int flagMergeThr, const int rasterizationFlags, const int* rasterizationMasks) //UE
 {
 	rcAssert(ctx);
 
@@ -1038,7 +1038,7 @@ void rcRasterizeTriangle(rcContext* ctx, const float* v0, const float* v1, const
 
 	const float ics = 1.0f/solid.cs;
 	const float ich = 1.0f/solid.ch;
-	rasterizeTri(v0, v1, v2, area, solid, solid.bmin, solid.bmax, solid.cs, ics, ich, flagMergeThr, rasterizationFlags, rasterizationMasks); //UE4
+	rasterizeTri(v0, v1, v2, area, solid, solid.bmin, solid.bmax, solid.cs, ics, ich, flagMergeThr, rasterizationFlags, rasterizationMasks); //UE
 
 	ctx->stopTimer(RC_TIMER_RASTERIZE_TRIANGLES);
 }
@@ -1050,7 +1050,7 @@ void rcRasterizeTriangle(rcContext* ctx, const float* v0, const float* v1, const
 /// @see rcHeightfield
 void rcRasterizeTriangles(rcContext* ctx, const float* verts, const int /*nv*/,
 						  const int* tris, const unsigned char* areas, const int nt,
-						  rcHeightfield& solid, const int flagMergeThr, const int rasterizationFlags, const int* rasterizationMasks) //UE4
+						  rcHeightfield& solid, const int flagMergeThr, const int rasterizationFlags, const int* rasterizationMasks) //UE
 {
 	if (ctx)
 		ctx->startTimer(RC_TIMER_RASTERIZE_TRIANGLES);
@@ -1064,7 +1064,7 @@ void rcRasterizeTriangles(rcContext* ctx, const float* verts, const int /*nv*/,
 		const float* v1 = &verts[tris[i*3+1]*3];
 		const float* v2 = &verts[tris[i*3+2]*3];
 		// Rasterize.
-		rasterizeTri(v0, v1, v2, areas[i], solid, solid.bmin, solid.bmax, solid.cs, ics, ich, flagMergeThr, rasterizationFlags, rasterizationMasks); //UE4
+		rasterizeTri(v0, v1, v2, areas[i], solid, solid.bmin, solid.bmax, solid.cs, ics, ich, flagMergeThr, rasterizationFlags, rasterizationMasks); //UE
 	}
 	
 	if (ctx)
@@ -1078,7 +1078,7 @@ void rcRasterizeTriangles(rcContext* ctx, const float* verts, const int /*nv*/,
 /// @see rcHeightfield
 void rcRasterizeTriangles(rcContext* ctx, const float* verts, const int /*nv*/,
 						  const unsigned short* tris, const unsigned char* areas, const int nt,
-						  rcHeightfield& solid, const int flagMergeThr, const int rasterizationFlags, const int* rasterizationMasks) //UE4
+						  rcHeightfield& solid, const int flagMergeThr, const int rasterizationFlags, const int* rasterizationMasks) //UE
 {
 	if (ctx)
 		ctx->startTimer(RC_TIMER_RASTERIZE_TRIANGLES);
@@ -1092,7 +1092,7 @@ void rcRasterizeTriangles(rcContext* ctx, const float* verts, const int /*nv*/,
 		const float* v1 = &verts[tris[i*3+1]*3];
 		const float* v2 = &verts[tris[i*3+2]*3];
 		// Rasterize.
-		rasterizeTri(v0, v1, v2, areas[i], solid, solid.bmin, solid.bmax, solid.cs, ics, ich, flagMergeThr, rasterizationFlags, rasterizationMasks); //UE4
+		rasterizeTri(v0, v1, v2, areas[i], solid, solid.bmin, solid.bmax, solid.cs, ics, ich, flagMergeThr, rasterizationFlags, rasterizationMasks); //UE
 	}
 	
 	if (ctx)
@@ -1105,7 +1105,7 @@ void rcRasterizeTriangles(rcContext* ctx, const float* verts, const int /*nv*/,
 ///
 /// @see rcHeightfield
 void rcRasterizeTriangles(rcContext* ctx, const float* verts, const unsigned char* areas, const int nt,
-						  rcHeightfield& solid, const int flagMergeThr, const int rasterizationFlags, const int* rasterizationMasks) //UE4
+						  rcHeightfield& solid, const int flagMergeThr, const int rasterizationFlags, const int* rasterizationMasks) //UE
 {
 	if (ctx)
 		ctx->startTimer(RC_TIMER_RASTERIZE_TRIANGLES);
@@ -1119,7 +1119,7 @@ void rcRasterizeTriangles(rcContext* ctx, const float* verts, const unsigned cha
 		const float* v1 = &verts[(i*3+1)*3];
 		const float* v2 = &verts[(i*3+2)*3];
 		// Rasterize.
-		rasterizeTri(v0, v1, v2, areas[i], solid, solid.bmin, solid.bmax, solid.cs, ics, ich, flagMergeThr, rasterizationFlags, rasterizationMasks); //UE4
+		rasterizeTri(v0, v1, v2, areas[i], solid, solid.bmin, solid.bmax, solid.cs, ics, ich, flagMergeThr, rasterizationFlags, rasterizationMasks); //UE
 	}
 	
 	if (ctx)
diff --git a/Engine/Source/Runtime/Navmesh/Private/Recast/RecastRegion.cpp b/Engine/Source/Runtime/Navmesh/Private/Recast/RecastRegion.cpp
--- a/Engine/Source/Runtime/Navmesh/Private/Recast/RecastRegion.cpp
+++ b/Engine/Source/Runtime/Navmesh/Private/Recast/RecastRegion.cpp
@@ -1366,7 +1366,7 @@ bool rcBuildRegionsChunky(rcContext* ctx, rcCompactHeightfield& chf,
 	return true;
 }
 
-/// @UE4: rcBuildRegions is now split into two functions: gathering regions and merging them
+/// @UE: rcBuildRegions is now split into two functions: gathering regions and merging them
 ///       it allows reusing existing code for building layer set (flood fill based partitioning)
 ///		  
 ///		  spanBuf4 is temporary buffer, allocated and freed by caller (size = chf.spanCount * 4)
diff --git a/Engine/Source/Runtime/Navmesh/Public/DebugUtils/DetourDebugDraw.h b/Engine/Source/Runtime/Navmesh/Public/DebugUtils/DetourDebugDraw.h
--- a/Engine/Source/Runtime/Navmesh/Public/DebugUtils/DetourDebugDraw.h
+++ b/Engine/Source/Runtime/Navmesh/Public/DebugUtils/DetourDebugDraw.h
@@ -45,14 +45,14 @@ NAVMESH_API void duDebugDrawNavMeshPortals(struct duDebugDraw* dd, const dtNavMe
 NAVMESH_API void duDebugDrawNavMeshPolysWithFlags(struct duDebugDraw* dd, const dtNavMesh& mesh, const unsigned short polyFlags, const unsigned int col);
 NAVMESH_API void duDebugDrawNavMeshPoly(struct duDebugDraw* dd, const dtNavMesh& mesh, dtPolyRef ref, const unsigned int col);
 
-//@UE4 BEGIN
+//@UE BEGIN
 #if WITH_NAVMESH_CLUSTER_LINKS
 NAVMESH_API void duDebugDrawNavMeshClusters(struct duDebugDraw* dd, const dtNavMesh& mesh);
 NAVMESH_API void duDebugDrawNavMeshCluster(struct duDebugDraw* dd, const dtNavMesh& mesh, dtClusterRef ref, const unsigned int col);
 
 NAVMESH_API void duDebugDrawTileCacheClusters(duDebugDraw* dd, const struct dtTileCacheClusterSet& lclusters);
 #endif // WITH_NAVMESH_CLUSTER_LINKS
-//@UE4 END
+//@UE END
 
 NAVMESH_API void duDebugDrawTileCacheLayerAreas(struct duDebugDraw* dd, const dtTileCacheLayer& layer, const float cs, const float ch);
 NAVMESH_API void duDebugDrawTileCacheLayerDistance(struct duDebugDraw* dd, const dtTileCacheLayer& layer,
diff --git a/Engine/Source/Runtime/Navmesh/Public/Detour/DetourCommon.h b/Engine/Source/Runtime/Navmesh/Public/Detour/DetourCommon.h
--- a/Engine/Source/Runtime/Navmesh/Public/Detour/DetourCommon.h
+++ b/Engine/Source/Runtime/Navmesh/Public/Detour/DetourCommon.h
@@ -486,7 +486,7 @@ inline void dtSwapEndian(float* v)
 void dtRandomPointInConvexPoly(const float* pts, const int npts, float* areas,
 							   const float s, const float t, float* out);
 
-// @UE4 BEGIN
+// @UE BEGIN
 enum dtRotation
 {
 	DT_ROTATE_0, 
@@ -525,7 +525,7 @@ void dtRotate90(float* dest, const float* v, const float* center, const dtRotati
 ///  @param[in]		center	The center point. [(x, y, z)]
 ///  @param[in]		rot		The rotation enum value.
 void dtRotate90(unsigned short* dest, const unsigned short* v, const unsigned short* center, const dtRotation rot);
-// @UE4 END
+// @UE END
 
 /// @}
 
diff --git a/Engine/Source/Runtime/Navmesh/Public/Detour/DetourNavMesh.h b/Engine/Source/Runtime/Navmesh/Public/Detour/DetourNavMesh.h
--- a/Engine/Source/Runtime/Navmesh/Public/Detour/DetourNavMesh.h
+++ b/Engine/Source/Runtime/Navmesh/Public/Detour/DetourNavMesh.h
@@ -36,23 +36,23 @@
 #if USE_64BIT_ADDRESS
 
 #if defined(__LP64__)
-// LP64 (Linux/OS X): UE4 will define its uint64 type as "unsigned long long" so we need to match this
-typedef unsigned long long UE4Type_uint64;
+// LP64 (Linux/OS X): UE will define its uint64 type as "unsigned long long" so we need to match this
+typedef unsigned long long UEType_uint64;
 #else
 #include <stdint.h>
-typedef uint64_t UE4Type_uint64;
+typedef uint64_t UEType_uint64;
 #endif
 
 /// A handle to a polygon within a navigation mesh tile.
 /// @ingroup detour
-typedef UE4Type_uint64 dtPolyRef;
+typedef UEType_uint64 dtPolyRef;
 
 /// A handle to a tile within a navigation mesh.
 /// @ingroup detour
-typedef UE4Type_uint64 dtTileRef;
+typedef UEType_uint64 dtTileRef;
 
 /// A handle to a cluster within a navigation mesh tile.
-typedef UE4Type_uint64 dtClusterRef;
+typedef UEType_uint64 dtClusterRef;
 
 #else
 
@@ -103,17 +103,17 @@ static const unsigned int DT_NULL_LINK = 0xffffffff;
 /// A flag that indicates that an off-mesh connection can be traversed in both directions. (Is bidirectional.)
 static const unsigned char DT_OFFMESH_CON_BIDIR = 0x01;
 
-//@UE4 BEGIN
+//@UE BEGIN
 static const unsigned char DT_OFFMESH_CON_POINT = 0x02;
 static const unsigned char DT_OFFMESH_CON_SEGMENT = 0x04;
 static const unsigned char DT_OFFMESH_CON_CHEAPAREA = 0x08;
-//@UE4 END
+//@UE END
 
 /// The maximum number of user defined area ids.
 /// @ingroup detour
 static const int DT_MAX_AREAS = 64;
 
-//@UE4 BEGIN
+//@UE BEGIN
 /// Navmesh tiles' salt will have at least this much bits
 static const int DT_MIN_SALT_BITS = 5;
 static const int DT_SALT_BASE = 1;
@@ -140,7 +140,7 @@ static const unsigned char DT_CLINK_VALID_BCK = 0x02;
 /// Index of first cluster link within tile
 static const unsigned int DT_CLINK_FIRST = 0x80000000;
 #endif // WITH_NAVMESH_CLUSTER_LINKS
-//@UE4 END
+//@UE END
 
 /// Tile flags used for various functions and fields.
 /// For an example, see dtNavMesh::addTile().
@@ -172,12 +172,12 @@ enum dtPolyTypes
 	/// The polygon is an off-mesh connection consisting of two vertices.
 	DT_POLYTYPE_OFFMESH_POINT = 1,
 
-	//@UE4 BEGIN
+	//@UE BEGIN
 #if WITH_NAVMESH_SEGMENT_LINKS
 	/// The polygon is an off-mesh connection consisting of four vertices.
 	DT_POLYTYPE_OFFMESH_SEGMENT = 2,
 #endif // WITH_NAVMESH_SEGMENT_LINKS
-	//@UE4 END
+	//@UE END
 };
 
 
@@ -250,7 +250,7 @@ struct dtBVNode
 	int i;							///< The node's index. (Negative for escape sequence.)
 };
 
-//@UE4 BEGIN
+//@UE BEGIN
 #if WITH_NAVMESH_SEGMENT_LINKS
 struct dtOffMeshSegmentConnection
 {
@@ -287,7 +287,7 @@ struct dtOffMeshSegmentConnection
 	inline bool getBiDirectional() const { return (flags & 0x80) != 0; }
 };
 #endif // WITH_NAVMESH_SEGMENT_LINKS
-//@UE4 END
+//@UE END
 
 /// Defines an navigation mesh off-mesh connection within a dtMeshTile object.
 /// An off-mesh connection is a user defined traversable connection made up to two vertices.
@@ -299,10 +299,10 @@ struct dtOffMeshConnection
 	/// The radius of the endpoints. [Limit: >= 0]
 	float rad;		
 
-	//@UE4 BEGIN
+	//@UE BEGIN
 	/// The snap height of endpoints (less than 0 = use step height)
 	float height;
-	//@UE4 END
+	//@UE END
 
 	/// The id of the offmesh connection. (User assigned when the navigation mesh is built.)
 	unsigned int userId;
@@ -316,7 +316,7 @@ struct dtOffMeshConnection
 	/// Link flags. 
 	unsigned char flags;
 
-	//@UE4 BEGIN
+	//@UE BEGIN
 	/// Sets link flags
 	inline void setFlags(unsigned char conFlags)
 	{
@@ -329,10 +329,10 @@ struct dtOffMeshConnection
 
 	/// Gets the link snap mode
 	inline bool getSnapToCheapestArea() const { return (flags & 0x40) != 0; }
-	//@UE4 END
+	//@UE END
 };
 
-//@UE4 BEGIN
+//@UE BEGIN
 #if WITH_NAVMESH_CLUSTER_LINKS
 /// Cluster of polys
 struct dtCluster
@@ -350,7 +350,7 @@ struct dtClusterLink
 	unsigned char flags;			///< Link traversing data
 };
 #endif // WITH_NAVMESH_CLUSTER_LINKS
-//@UE4 END
+//@UE END
 
 /// Provides high level information related to a dtMeshTile object.
 /// @ingroup detour
@@ -375,7 +375,7 @@ struct dtMeshHeader
 	int offMeshConCount;		///< The number of point type off-mesh connections.
 	int offMeshBase;			///< The index of the first polygon which is an point type off-mesh connection.
 
-	//@UE4 BEGIN
+	//@UE BEGIN
 #if WITH_NAVMESH_SEGMENT_LINKS
 	int offMeshSegConCount;		///< The number of segment type off-mesh connections.
 	int offMeshSegPolyBase;		///< The index of the first polygon which is an segment type off-mesh connection
@@ -385,7 +385,7 @@ struct dtMeshHeader
 #if WITH_NAVMESH_CLUSTER_LINKS
 	int clusterCount;			///< Number of clusters
 #endif // WITH_NAVMESH_CLUSTER_LINKS
-	//@UE4 END
+	//@UE END
 
 	float walkableHeight;		///< The height of the agents using the tile.
 	float walkableRadius;		///< The radius of the agents using the tile.
@@ -427,7 +427,7 @@ struct dtMeshTile
 	int flags;								///< Tile flags. (See: #dtTileFlags)
 	dtMeshTile* next;						///< The next free tile, or the next tile in the spatial grid.
 
-	//@UE4 BEGIN
+	//@UE BEGIN
 #if WITH_NAVMESH_SEGMENT_LINKS
 	dtOffMeshSegmentConnection* offMeshSeg;		///< The tile off-mesh connections. [Size: dtMeshHeader::offMeshSegConCount]
 #endif // WITH_NAVMESH_SEGMENT_LINKS
@@ -442,7 +442,7 @@ struct dtMeshTile
 
 	dtChunkArray<dtLink> dynamicLinksO;			///< Dynamic links array (indices starting from dtMeshHeader::maxLinkCount)
 	unsigned int dynamicFreeListO;				///< Index of the next free dynamic link
-	//@UE4 END
+	//@UE END
 };
 
 /// Configuration parameters used to define multi-tile navigation meshes.
@@ -519,13 +519,13 @@ public:
 	/// @return The tile, or null if the tile does not exist.
 	const dtMeshTile* getTileAt(const int x, const int y, const int layer) const;
 
-// @UE4 BEGIN
+// @UE BEGIN
 	/// Gets number of tiles at the specified grid location. (All layers.)
 	///  @param[in]		x			The tile's x-location. (x, y)
 	///  @param[in]		y			The tile's y-location. (x, y)
 	/// @return The number of tiles in grid.
 	int getTileCountAt(const int x, const int y) const;
-// @UE4 END
+// @UE END
 
 	/// Gets all tiles at the specified grid location. (All layers.)
 	///  @param[in]		x			The tile's x-location. (x, y)
@@ -610,7 +610,7 @@ public:
 	///	 @param[in] newArea	Area code to apply
 	void updateOffMeshConnectionByUserId(unsigned int userId, unsigned char newArea, unsigned short newFlags);
 
-	//@UE4 BEGIN
+	//@UE BEGIN
 #if WITH_NAVMESH_SEGMENT_LINKS
 	/// Gets the specified off-mesh connection: segment type
 	///  @param[in]	ref		The polygon reference of the off-mesh connection.
@@ -622,7 +622,7 @@ public:
 	///	 @param[in] newArea	Area code to apply
 	void updateOffMeshSegmentConnectionByUserId(unsigned int userId, unsigned char newArea, unsigned short newFlags);
 #endif // WITH_NAVMESH_SEGMENT_LINKS
-	//@UE4 END
+	//@UE END
 
 	/// @}
 
@@ -744,7 +744,7 @@ public:
 		return decodePolyIdTile(ref);
 	}
 
-	//@UE4 BEGIN 
+	//@UE BEGIN 
 #if WITH_NAVMESH_CLUSTER_LINKS
 	/// Extracts the cluster's index (within its tile) from the specified cluster reference.
 	///  @note This function is generally meant for internal use only.
@@ -824,12 +824,12 @@ public:
 	{
 		return (unsigned int)(tile - m_tiles);
 	}
-	//@UE4 END
+	//@UE END
 	/// @}
 
 private:
 
-	// [UE4] result struct for findConnectingPolys
+	// [UE] result struct for findConnectingPolys
 	struct FConnectingPolyData
 	{
 		float min;
@@ -848,7 +848,7 @@ private:
 	int getNeighbourTilesAt(const int x, const int y, const int side,
 							dtMeshTile** tiles, const int maxTiles) const;
 	
-	/// [UE4] Returns all polygons in neighbour tile based on portal defined by the segment.
+	/// [UE] Returns all polygons in neighbour tile based on portal defined by the segment.
 	int findConnectingPolys(const float* va, const float* vb,
 		const dtMeshTile* fromTile, int fromPolyIdx,
 		const dtMeshTile* tile, int side,
@@ -867,7 +867,7 @@ private:
 	/// Removes external links at specified side.
 	void unconnectExtLinks(dtMeshTile* tile, dtMeshTile* target);
 	
-	//@UE4 BEGIN
+	//@UE BEGIN
 #if WITH_NAVMESH_CLUSTER_LINKS
 	/// Try to connect clusters
 	void connectClusterLink(dtMeshTile* tile0, unsigned int cluster0,
@@ -877,7 +877,7 @@ private:
 	/// Removes cluster links at specified side.
 	void unconnectClusterLinks(dtMeshTile* tile, dtMeshTile* target);
 #endif // WITH_NAVMESH_CLUSTER_LINKS
-	//@UE4 END
+	//@UE END
 
 	// TODO: These methods are duplicates from dtNavMeshQuery, but are needed for off-mesh connection finding.
 	
@@ -921,7 +921,7 @@ NAVMESH_API dtNavMesh* dtAllocNavMesh();
 ///  @ingroup detour
 NAVMESH_API void dtFreeNavMesh(dtNavMesh* navmesh);
 
-// @UE4 BEGIN: helper for reading tiles
+// @UE BEGIN: helper for reading tiles
 struct ReadTilesHelper
 {
 	static const int MaxTiles = 32;
@@ -953,7 +953,7 @@ struct ReadTilesHelper
 		return AllocatedTiles;
 	}
 };
-// @UE4 END
+// @UE END
 
 #endif // DETOURNAVMESH_H
 
diff --git a/Engine/Source/Runtime/Navmesh/Public/Detour/DetourNavMeshBuilder.h b/Engine/Source/Runtime/Navmesh/Public/Detour/DetourNavMeshBuilder.h
--- a/Engine/Source/Runtime/Navmesh/Public/Detour/DetourNavMeshBuilder.h
+++ b/Engine/Source/Runtime/Navmesh/Public/Detour/DetourNavMeshBuilder.h
@@ -75,7 +75,7 @@ struct dtNavMeshCreateParams
 	int polyCount;							///< Number of polygons in the mesh. [Limit: >= 1]
 	int nvp;								///< Number maximum number of vertices per polygon. [Limit: >= 3]
 
-	// @UE4 BEGIN
+	// @UE BEGIN
 #if WITH_NAVMESH_CLUSTER_LINKS
 	/// @}
 	/// @name Cluster Attributes
@@ -83,7 +83,7 @@ struct dtNavMeshCreateParams
 	unsigned short* polyClusters;			///< Cluster Id for each polygon [Size: #polyCount]
 	unsigned short clusterCount;			///< Number of unique clusters
 #endif // WITH_NAVMESH_CLUSTER_LINKS
-	// @UE4 END
+	// @UE END
 
 	/// @}
 	/// @name Height Detail Attributes (Optional)
@@ -166,7 +166,7 @@ NAVMESH_API bool dtNavMeshHeaderSwapEndian(unsigned char* data, const int dataSi
 ///  @param[in]		dataSize	The size of the data array.
 NAVMESH_API bool dtNavMeshDataSwapEndian(unsigned char* data, const int dataSize);
 
-// @UE4 BEGIN
+// @UE BEGIN
 /// Offset and rotate around center the data in the tile
 ///  @param[in,out]	data			Data of the tile mesh. (See: #dtCreateNavMeshData)
 ///  @param[in]		dataSize		Data size of the tile mesh.
@@ -186,7 +186,7 @@ NAVMESH_API bool dtTransformTileData(unsigned char* data, const int dataSize, co
 ///  @param[out]	deltaX			Offset X in tile coordinates.
 ///  @param[out]	deltaY			Offset Y in tile coordinates.
 NAVMESH_API void dtComputeTileOffsetFromRotation(const float* position, const float* rotationCenter, const float rotationDeg, const float tileWidth, const float tileHeight, int& deltaX, int& deltaY);
-// @UE4 END
+// @UE END
 
 #endif // DETOURNAVMESHBUILDER_H
 
diff --git a/Engine/Source/Runtime/Navmesh/Public/Detour/DetourNavMeshQuery.h b/Engine/Source/Runtime/Navmesh/Public/Detour/DetourNavMeshQuery.h
--- a/Engine/Source/Runtime/Navmesh/Public/Detour/DetourNavMeshQuery.h
+++ b/Engine/Source/Runtime/Navmesh/Public/Detour/DetourNavMeshQuery.h
@@ -27,11 +27,11 @@
 #include "Detour/DetourNavMesh.h"
 #include "Detour/DetourCommon.h"
 
-//@UE4 BEGIN
+//@UE BEGIN
 #define WITH_FIXED_AREA_ENTERING_COST 1
 
 #define DT_UNWALKABLE_POLY_COST FLT_MAX
-//@UE4 END
+//@UE END
 
 // By default dtQueryFilter will use virtual calls.
 // On certain platforms indirect or virtual function call is expensive. The default
@@ -40,7 +40,7 @@
 // To avoid virtual calls create dtQueryFilter with inIsVirtual = false.
 
 // Special link filter is custom filter run only for offmesh links with assigned UserId
-// Used by smart navlinks in UE4
+// Used by smart navlinks in UE
 //
 struct NAVMESH_API dtQuerySpecialLinkFilter
 {
@@ -54,7 +54,7 @@ struct NAVMESH_API dtQuerySpecialLinkFilter
 	virtual void initialize() {}
 };
 
-// [UE4: moved all filter variables to struct, DO NOT mess with virtual functions here!]
+// [UE: moved all filter variables to struct, DO NOT mess with virtual functions here!]
 struct NAVMESH_API dtQueryFilterData
 {
 	float m_areaCost[DT_MAX_AREAS];		///< Cost per area type. (Used by default implementation.)
@@ -68,13 +68,13 @@ struct NAVMESH_API dtQueryFilterData
 	unsigned short m_excludeFlags;		///< Flags for polygons that should not be visited. (Used by default implementation.)
 
 	bool m_isBacktracking;
-	//@UE4 BEGIN
+	//@UE BEGIN
 	/// whether to ignore neighbour nodes that have already been visited.
 	/// FALSE by default but you might want to set it to true if using heuristic
 	/// scale much larger than 1 or experiencing cycles in our paths for any other
 	/// reason
 	bool m_shouldIgnoreClosedNodes;
-	//@UE4 END
+	//@UE END
 
 	dtQueryFilterData();
 	
@@ -196,7 +196,7 @@ public:
 	///  @param[in]		cost	The new cost of traversing the area.
 	inline void setAreaCost(const int i, const float cost) { data.m_areaCost[i] = cost; data.lowestAreaCost = dtMin(data.lowestAreaCost, cost); }
 
-//@UE4 BEGIN
+//@UE BEGIN
 	inline const float* getAllAreaCosts() const { return data.m_areaCost; }
 	
 #if WITH_FIXED_AREA_ENTERING_COST
@@ -244,7 +244,7 @@ public:
 	/// Retrieves information whether this filter allows reopening closed nodes
 	///  @returns should consider reopening nodes already on closed list
 	inline bool getShouldIgnoreClosedNodes() const { return data.m_shouldIgnoreClosedNodes; }
-//@UE4 END
+//@UE END
 
 	/// Returns the include flags for the filter.
 	/// Any polygons that include one or more of these flags will be
@@ -333,7 +333,7 @@ public:
 	/// @returns The status flags for the query.
 	dtStatus init(const dtNavMesh* nav, const int maxNodes, dtQuerySpecialLinkFilter* linkFilter = 0);
 
-	/// UE4: updates special link filter for this query
+	/// UE: updates special link filter for this query
 	void updateLinkFilter(dtQuerySpecialLinkFilter* linkFilter);
 	
 	/// @name Standard Pathfinding Functions
@@ -344,12 +344,12 @@ public:
 	///  @param[in]		endRef		The reference id of the end polygon.
 	///  @param[in]		startPos	A position within the start polygon. [(x, y, z)]
 	///  @param[in]		endPos		A position within the end polygon. [(x, y, z)]
-	///  @param[in]		costLimit	Cost limit of nodes allowed to be added to the open list	//@UE4
+	///  @param[in]		costLimit	Cost limit of nodes allowed to be added to the open list	//@UE
 	///  @param[in]		filter		The polygon filter to apply to the query.
 	///  @param[out]	result		Results for path corridor, fills in refs and costs for each poly from start to end
 	///	 @param[out]	totalCost			If provided will get filled will total cost of path
 	dtStatus findPath(dtPolyRef startRef, dtPolyRef endRef,
-					  const float* startPos, const float* endPos, const float costLimit, //@UE4
+					  const float* startPos, const float* endPos, const float costLimit, //@UE
 					  const dtQueryFilter* filter,
 					  dtQueryResult& result, float* totalCost) const;
 	
@@ -378,11 +378,11 @@ public:
 	///  @param[in]		endRef		The reference id of the end polygon.
 	///  @param[in]		startPos	A position within the start polygon. [(x, y, z)]
 	///  @param[in]		endPos		A position within the end polygon. [(x, y, z)]
-	///  @param[in]		costLimit	Cost limit of nodes allowed to be added to the open list	//@UE4
+	///  @param[in]		costLimit	Cost limit of nodes allowed to be added to the open list	//@UE
 	///  @param[in]		filter		The polygon filter to apply to the query.
 	/// @returns The status flags for the query.
 	dtStatus initSlicedFindPath(dtPolyRef startRef, dtPolyRef endRef,
-								const float* startPos, const float* endPos, const float costLimit, //@UE4
+								const float* startPos, const float* endPos, const float costLimit, //@UE
 								const dtQueryFilter* filter);
 
 	/// Updates an in-progress sliced path query.
@@ -450,7 +450,7 @@ public:
 								  dtPolyRef* resultRef, dtPolyRef* resultParent, float* resultCost,
 								  int* resultCount, const int maxResult) const;
 	
-	//@UE4 BEGIN
+	//@UE BEGIN
 	/// Finds the polygons along the navigation graph that are no more than given path length away from centerPos.
 	///  @param[in]		startRef		The reference id of the polygon where the search starts.
 	///  @param[in]		centerPos		The center of the search circle. [(x, y, z)]
@@ -504,7 +504,7 @@ public:
 	dtStatus findNearestContainingPoly(const float* center, const float* extents,
 									   const dtQueryFilter* filter,
 									   dtPolyRef* nearestRef, float* nearestPt) const;
-	//@UE4 END
+	//@UE END
 
 	/// Finds polygons that overlap the search box.
 	///  @param[in]		center		The center of the search box. [(x, y, z)]
@@ -534,7 +534,7 @@ public:
 									dtPolyRef* resultRef, dtPolyRef* resultParent,
 									int* resultCount, const int maxResult) const;
 
-	/// [UE4] Finds the wall segments in local neighbourhood
+	/// [UE] Finds the wall segments in local neighbourhood
 	dtStatus findWallsInNeighbourhood(dtPolyRef startRef, const float* centerPos, const float radius,
 									  const dtQueryFilter* filter, 
 									  dtPolyRef* neiRefs, int* neiCount, const int maxNei,
@@ -623,7 +623,7 @@ public:
 										 const dtQueryFilter* filter, float (*frand)(),
 										 dtPolyRef* randomRef, float* randomPt) const;
 
-	//@UE4 BEGIN
+	//@UE BEGIN
 #if WITH_NAVMESH_CLUSTER_LINKS
 	/// Check if there is a path from start polygon to the end polygon using cluster graph
 	/// (cheap, does not care about link costs)
@@ -645,7 +645,7 @@ public:
 	/// @returns The status flags for the query.
 	dtStatus getPolyCluster(dtPolyRef polyRef, dtClusterRef& clusterRef) const;
 #endif // WITH_NAVMESH_CLUSTER_LINKS
-//@UE4 END
+//@UE END
 	
 	/// Finds the closest point on the specified polygon.
 	///  @param[in]		ref			The reference id of the polygon.
@@ -734,7 +734,7 @@ private:
 	/// Returns projected point on polygon.
 	dtStatus projectedPointOnPolyInTile(const dtMeshTile* tile, const dtPoly* poly, const float* pos, float* projected) const;
 	
-	//@UE4 BEGIN
+	//@UE BEGIN
 	// exposing function to be able to generate navigation corridors as sequence of point pairs
 public:
 	/// Returns portal points between two polygons.
@@ -749,7 +749,7 @@ public:
 	dtStatus getEdgeMidPoint(dtPolyRef from, const dtPoly* fromPoly, const dtMeshTile* fromTile,
 							 dtPolyRef to, const dtPoly* toPoly, const dtMeshTile* toTile,
 							 float* mid) const;
-	//@UE4 END
+	//@UE END
 
 	// Appends vertex to a straight path
 	dtStatus appendVertex(const float* pos, const unsigned char flags, const dtPolyRef ref,
@@ -785,7 +785,7 @@ private:
 		float lastBestNodeCost;
 		dtPolyRef startRef, endRef;
 		float startPos[3], endPos[3];
-		float costLimit; 					//@UE4 ///< Cost limit of nodes allowed to be added to the open list
+		float costLimit; 					//@UE ///< Cost limit of nodes allowed to be added to the open list
 		const dtQueryFilter* filter;
 	};
 	dtQueryData m_query;				///< Sliced query state.
diff --git a/Engine/Source/Runtime/Navmesh/Public/Detour/DetourNode.h b/Engine/Source/Runtime/Navmesh/Public/Detour/DetourNode.h
--- a/Engine/Source/Runtime/Navmesh/Public/Detour/DetourNode.h
+++ b/Engine/Source/Runtime/Navmesh/Public/Detour/DetourNode.h
@@ -82,22 +82,22 @@ public:
 	}
 	
 	inline int getMaxNodes() const { return m_maxNodes; }
-	//@UE4 BEGIN
+	//@UE BEGIN
 	// If using a shared query instance it's possible that m_maxNodes is greater
 	// than pool size requested by callee. There's no point in reallocating the
 	// pool so we artificially limit the number of available nodes
 	inline int getMaxRuntimeNodes() const { return m_maxRuntimeNodes; }
-	//@UE4 END
+	//@UE END
 	inline int getNodeCount() const { return m_nodeCount; }
 	
 	inline int getHashSize() const { return m_hashSize; }
 	inline dtNodeIndex getFirst(int bucket) const { return m_first[bucket]; }
 	inline dtNodeIndex getNext(int i) const { return m_next[i]; }
 
-	//@UE4 BEGIN
+	//@UE BEGIN
 	// overrides m_maxNodes for runtime purposes
 	inline void setMaxRuntimeNodes(const int newMaxRuntimeNodes) { m_maxRuntimeNodes = newMaxRuntimeNodes; }
-	//@UE4 END
+	//@UE END
 
 private:
 	
@@ -106,9 +106,9 @@ private:
 	dtNodeIndex* m_next;
 	const int m_maxNodes;
 	const int m_hashSize;
-	//@UE4 BEGIN
+	//@UE BEGIN
 	int m_maxRuntimeNodes;
-	//@UE4 END
+	//@UE END
 	int m_nodeCount;
 };
 
diff --git a/Engine/Source/Runtime/Navmesh/Public/DetourCrowd/DetourCrowd.h b/Engine/Source/Runtime/Navmesh/Public/DetourCrowd/DetourCrowd.h
--- a/Engine/Source/Runtime/Navmesh/Public/DetourCrowd/DetourCrowd.h
+++ b/Engine/Source/Runtime/Navmesh/Public/DetourCrowd/DetourCrowd.h
@@ -51,7 +51,7 @@ static const int DT_CROWDAGENT_MAX_CORNERS = 4;
 ///		 dtCrowdAgentParams::obstacleAvoidanceType
 static const int DT_CROWD_MAX_OBSTAVOIDANCE_PARAMS = 8;
 
-/// UE4: The maximum number of unique filters used by crowd agents 
+/// UE: The maximum number of unique filters used by crowd agents 
 static const int DT_CROWD_MAX_FILTERS = 16;
 
 /// Provides neighbor data for agents managed by the crowd.
@@ -70,7 +70,7 @@ enum CrowdAgentState
 	DT_CROWDAGENT_STATE_INVALID,		///< The agent is not in a valid state.
 	DT_CROWDAGENT_STATE_WALKING,		///< The agent is traversing a normal navigation mesh polygon.
 	DT_CROWDAGENT_STATE_OFFMESH,		///< The agent is traversing an off-mesh connection.
-	DT_CROWDAGENT_STATE_WAITING,		///< [UE4] The agent is waiting for external movement to finish
+	DT_CROWDAGENT_STATE_WAITING,		///< [UE] The agent is waiting for external movement to finish
 };
 
 /// Configuration parameters for a crowd agent.
@@ -80,7 +80,7 @@ struct dtCrowdAgentParams
 	/// User defined data attached to the agent.
 	void* userData;
 
-	/// UE4: special link filter used by this agent
+	/// UE: special link filter used by this agent
 	TSharedPtr<dtQuerySpecialLinkFilter> linkFilter;
 
 	float radius;						///< Agent radius. [Limit: >= 0]
@@ -96,14 +96,14 @@ struct dtCrowdAgentParams
 	/// How aggresive the agent manager should be at avoiding collisions with this agent. [Limit: >= 0]
 	float separationWeight;
 
-	/// [UE4] Mutliplier for avoidance velocities
+	/// [UE] Mutliplier for avoidance velocities
 	float avoidanceQueryMultiplier;
 
-	/// [UE4] Groups flags attached to the agent
+	/// [UE] Groups flags attached to the agent
 	unsigned int avoidanceGroup;
-	/// [UE4] Avoid agents when they group is matching mask
+	/// [UE] Avoid agents when they group is matching mask
 	unsigned int groupsToAvoid;
-	/// [UE4] Don't avoid agents when they group is matching mask
+	/// [UE] Don't avoid agents when they group is matching mask
 	unsigned int groupsToIgnore;
 
 	/// Flags that impact steering behavior. (See: #UpdateFlags)
@@ -113,7 +113,7 @@ struct dtCrowdAgentParams
 	/// [Limits: 0 <= value <= #DT_CROWD_MAX_OBSTAVOIDANCE_PARAMS]
 	unsigned char obstacleAvoidanceType;	
 
-	/// UE4: Id of navigation filter used by this agent
+	/// UE: Id of navigation filter used by this agent
 	/// [Limits: 0 <= value <= #DT_CROWD_MAX_FILTERS]
 	unsigned char filter;
 };
@@ -204,12 +204,12 @@ enum UpdateFlags
 	DT_CROWD_SEPARATION				= 1 << 2,
 	DT_CROWD_OPTIMIZE_VIS			= 1 << 3,	///< Use #dtPathCorridor::optimizePathVisibility() to optimize the agent path.
 	DT_CROWD_OPTIMIZE_TOPO			= 1 << 4,	///< Use dtPathCorridor::optimizePathTopology() to optimize the agent path.
-	DT_CROWD_OPTIMIZE_VIS_MULTI		= 1 << 5,	///< [UE4] Multiple calls for optimizePathVisibility instead of checking last point
-	DT_CROWD_OFFSET_PATH			= 1 << 6,	///< [UE4] Offset path points from corners by agent radius
-	DT_CROWD_SLOWDOWN_AT_GOAL		= 1 << 7,	///< [UE4] Slowdown before reaching goal
+	DT_CROWD_OPTIMIZE_VIS_MULTI		= 1 << 5,	///< [UE] Multiple calls for optimizePathVisibility instead of checking last point
+	DT_CROWD_OFFSET_PATH			= 1 << 6,	///< [UE] Offset path points from corners by agent radius
+	DT_CROWD_SLOWDOWN_AT_GOAL		= 1 << 7,	///< [UE] Slowdown before reaching goal
 };
 
-// [UE4] Flags used by boundary segments (dtLocalBoundary::Segment)
+// [UE] Flags used by boundary segments (dtLocalBoundary::Segment)
 enum CrowdBoundaryFlags
 {
 	DT_CROWD_BOUNDARY_IGNORE = 1 << 0,
@@ -250,24 +250,24 @@ class NAVMESH_API dtCrowd
 	
 	float m_maxAgentRadius;
 
-	// [UE4] time between attempts to restore agents state
+	// [UE] time between attempts to restore agents state
 	float m_agentStateCheckInterval;
 
-	// [UE4] radius multiplier for offseting path around corners
+	// [UE] radius multiplier for offseting path around corners
 	float m_pathOffsetRadiusMultiplier;
 
-	// [UE4] separation filter
+	// [UE] separation filter
 	float m_separationDirFilter;
 
 	int m_velocitySampleCount;
 
 	dtNavMeshQuery* m_navquery;
 
-	// [UE4] if set, path visibility optimization can't leave current area type
+	// [UE] if set, path visibility optimization can't leave current area type
 	bool m_raycastSingleArea;
-	// [UE4] if set, offmesh connections won't be cut from corridor
+	// [UE] if set, offmesh connections won't be cut from corridor
 	bool m_keepOffmeshConnections;
-	// [UE4] if set, crowd agents will use early reach test
+	// [UE] if set, crowd agents will use early reach test
 	bool m_earlyReachTest;
 
 	void updateTopologyOptimization(dtCrowdAgent** agents, const int nagents, const float dt);
@@ -289,7 +289,7 @@ public:
 	/// @return True if the initialization succeeded.
 	bool init(const int maxAgents, const float maxAgentRadius, dtNavMesh* nav);
 
-	/// [UE4] Initializes the avoidance query.  
+	/// [UE] Initializes the avoidance query.  
 	///  @param[in]		maxNeighbors		The maximum number of processed neighbors
 	///  @param[in]		maxWalls			The maximum number of processed wall segments
 	///  @param[in]		maxCustomPatterns	The maximum number of custom sampling patterns
@@ -307,14 +307,14 @@ public:
 	/// @return The requested configuration.
 	const dtObstacleAvoidanceParams* getObstacleAvoidanceParams(const int idx) const;
 	
-	/// [UE4] Sets the shared avoidance sampling pattern for the specified index.
+	/// [UE] Sets the shared avoidance sampling pattern for the specified index.
 	///  @param[in]		idx			The index.
 	///  @param[in]		angles		radians from direction of desired velocity [Count: nsamples]
 	///  @param[in]		radii		normalized radii (0..1) for each sample [Count: nsamples]
 	///  @param[in]		nsamples	The number of samples
 	void setObstacleAvoidancePattern(int idx, const float* angles, const float* radii, int nsamples);
 
-	/// [UE4] Gets the shared avoidance sampling pattern for the specified index.
+	/// [UE] Gets the shared avoidance sampling pattern for the specified index.
 	///  @param[in]		idx			The index.
 	///  @param[in]		angles		radians from direction of desired velocity [Count: nsamples]
 	///  @param[in]		radii		normalized radii (0..1) for each sample [Count: nsamples]
@@ -334,7 +334,7 @@ public:
 	/// Adds a new agent to the crowd.
 	///  @param[in]		pos		The requested position of the agent. [(x, y, z)]
 	///  @param[in]		params	The configutation of the agent.
-	///  @param[in]		filter	[UE4] query filter used by agent
+	///  @param[in]		filter	[UE] query filter used by agent
 	/// @return The index of the agent in the agent pool. Or -1 if the agent could not be added.
 	int addAgent(const float* pos, const dtCrowdAgentParams& params, const dtQueryFilter* filter);
 
@@ -343,13 +343,13 @@ public:
 	///  @param[in]		params	The new agent configuration.
 	void updateAgentParameters(const int idx, const dtCrowdAgentParams& params);
 
-	/// [UE4] Updates the specified agent's query filter.
+	/// [UE] Updates the specified agent's query filter.
 	///  @param[in]		idx		The agent index. [Limits: 0 <= value < #getAgentCount()]
 	///  @param[in]		filter	The new filter.
 	/// @return True if the request was successfully submitted.
 	bool updateAgentFilter(const int idx, const dtQueryFilter* filter);
 
-	/// [UE4] Refresh state of agent, used after completing movement through offmesh links
+	/// [UE] Refresh state of agent, used after completing movement through offmesh links
 	///  @param[in]		idx		The agent index. [Limits: 0 <= value < #getAgentCount()]
 	///  @param[in]		repath	If set, agent will invalidate its path
 	void updateAgentState(const int idx, bool repath);
@@ -376,15 +376,15 @@ public:
 	/// @return True if the request was successfully reseted.
 	bool resetMoveTarget(const int idx);
 
-	/// [UE4] Switch to waiting state
+	/// [UE] Switch to waiting state
 	///  @param[in]		idx		The agent index. [Limits: 0 <= value < #getAgentCount()]
 	bool setAgentWaiting(const int idx);
 
-	/// [UE4] Switch to offmesh link state
+	/// [UE] Switch to offmesh link state
 	///  @param[in]		idx		The agent index. [Limits: 0 <= value < #getAgentCount()]
 	bool setAgentBackOnLink(const int idx);
 
-	/// [UE4] Resets agent's velocity
+	/// [UE] Resets agent's velocity
 	///  @param[in]		idx		The agent index. [Limits: 0 <= value < #getAgentCount()]
 	bool resetAgentVelocity(const int idx);
 
@@ -403,56 +403,56 @@ public:
 	///  @param[out]	debug	A debug object to load with debug information. [Opt]
 	void update(const float dt, dtCrowdAgentDebugInfo* debug);
 	
-	/// [UE4] Split update into several smaller components: path validity, path cache and path optimizations
+	/// [UE] Split update into several smaller components: path validity, path cache and path optimizations
 	/// @param[in]		dt		Delta time in seconds
 	/// @param[in]		nagents	Number of active agents
 	void updateStepPaths(const float dt, dtCrowdAgentDebugInfo* debug);
 
-	/// [UE4] Split update into several smaller components: neighbors and boundaries
+	/// [UE] Split update into several smaller components: neighbors and boundaries
 	/// @param[in]		dt		Delta time in seconds
 	/// @param[in]		nagents	Number of active agents
 	void updateStepProximityData(const float dt, dtCrowdAgentDebugInfo* debug);
 
-	/// [UE4] Split update into several smaller components: next corner for move, trigger offmesh links
+	/// [UE] Split update into several smaller components: next corner for move, trigger offmesh links
 	/// @param[in]		dt		Delta time in seconds
 	/// @param[in]		nagents	Number of active agents
 	void updateStepNextMovePoint(const float dt, dtCrowdAgentDebugInfo* debug);
 
-	/// [UE4] Split update into several smaller components: steering
+	/// [UE] Split update into several smaller components: steering
 	/// @param[in]		dt		Delta time in seconds
 	/// @param[in]		nagents	Number of active agents
 	void updateStepSteering(const float dt, dtCrowdAgentDebugInfo* debug);
 
-	/// [UE4] Split update into several smaller components: avoidance
+	/// [UE] Split update into several smaller components: avoidance
 	/// @param[in]		dt		Delta time in seconds
 	/// @param[in]		nagents	Number of active agents
 	void updateStepAvoidance(const float dt, dtCrowdAgentDebugInfo* debug);
 
-	/// [UE4] Split update into several smaller components: integrate velocities and handle collisions
+	/// [UE] Split update into several smaller components: integrate velocities and handle collisions
 	/// @param[in]		dt		Delta time in seconds
 	/// @param[in]		nagents	Number of active agents
 	void updateStepMove(const float dt, dtCrowdAgentDebugInfo* debug);
 
-	/// [UE4] Split update into several smaller components: corridor updates at new position
+	/// [UE] Split update into several smaller components: corridor updates at new position
 	/// @param[in]		dt		Delta time in seconds
 	/// @param[in]		nagents	Number of active agents
 	void updateStepCorridor(const float dt, dtCrowdAgentDebugInfo* debug);
 
-	/// [UE4] Split update into several smaller components: offmesh anims
+	/// [UE] Split update into several smaller components: offmesh anims
 	/// @param[in]		dt		Delta time in seconds
 	/// @param[in]		nagents	Number of active agents
 	void updateStepOffMeshAnim(const float dt, dtCrowdAgentDebugInfo* debug);
 
-	/// [UE4] Split update into several smaller components: offmesh link velocity (instead of playing animation)
+	/// [UE] Split update into several smaller components: offmesh link velocity (instead of playing animation)
 	/// @param[in]		dt		Delta time in seconds
 	/// @param[in]		nagents	Number of active agents
 	void updateStepOffMeshVelocity(const float dt, dtCrowdAgentDebugInfo* debug);
 
-	/// [UE4] Set time between attempts to restore agents state
+	/// [UE] Set time between attempts to restore agents state
 	/// @param[in]		t		Time in seconds
 	void setAgentCheckInterval(const float t);
 
-	/// [UE4] Set agent corridor, works only just after requesting move target
+	/// [UE] Set agent corridor, works only just after requesting move target
 	/// when agent didn't start any pathfinding operations yet
 	/// Use with caution!
 	/// @param[in]		idx		The agent index. [Limits: 0 <= value < #getAgentCount()]
@@ -460,36 +460,36 @@ public:
 	/// @param[in]		npath	The number of polygons in the path.
 	bool setAgentCorridor(const int idx, const dtPolyRef* path, const int npath);
 
-	/// [UE4] Set visibility optimization to use single area raycasts
+	/// [UE] Set visibility optimization to use single area raycasts
 	/// This will prevent from cutting through polys marked as different area
 	/// which could have been avoided in corridor's path
 	/// @param[in]		bEnable	New state of single area raycast mode
 	void setSingleAreaVisibilityOptimization(bool bEnable);
 
-	/// [UE4] Set offmesh connection pruning
+	/// [UE] Set offmesh connection pruning
 	/// This will allow removing offmesh connection poly ref from corridor
 	/// as soon as offmesh connection anim is triggered (default behavior)
 	void setPruneStartedOffmeshConnections(bool bRemoveFromCorridor);
 
-	/// [UE4]
+	/// [UE]
 	void setEarlyReachTestOptimization(bool bEnable);
 
-	/// [UE4] Set agent radius multiplier for offseting path from corners
+	/// [UE] Set agent radius multiplier for offseting path from corners
 	void setPathOffsetRadiusMultiplier(float RadiusMultiplier);
 
-	/// [UE4] Set separation filter param
+	/// [UE] Set separation filter param
 	void setSeparationFilter(float InFilter);
 
-	/// [UE4] Check if agent moved away from its path corridor
+	/// [UE] Check if agent moved away from its path corridor
 	bool isOutsideCorridor(const int idx) const;
 
 	/// Gets the filter used by the crowd.
-	///  @param[in]		idx		[UE4] The agent index. [Limits: 0 <= value < #getAgentCount()]
+	///  @param[in]		idx		[UE] The agent index. [Limits: 0 <= value < #getAgentCount()]
 	/// @return The filter used by the crowd.
 	const dtQueryFilter* getFilter(const int idx) const;
 
 	/// Gets the filter used by the crowd.
-	///  @param[in]		idx		[UE4] The agent index. [Limits: 0 <= value < #getAgentCount()]
+	///  @param[in]		idx		[UE] The agent index. [Limits: 0 <= value < #getAgentCount()]
 	/// @return The filter used by the crowd.
 	dtQueryFilter* getEditableFilter(const int idx);
 
diff --git a/Engine/Source/Runtime/Navmesh/Public/DetourCrowd/DetourLocalBoundary.h b/Engine/Source/Runtime/Navmesh/Public/DetourCrowd/DetourLocalBoundary.h
--- a/Engine/Source/Runtime/Navmesh/Public/DetourCrowd/DetourLocalBoundary.h
+++ b/Engine/Source/Runtime/Navmesh/Public/DetourCrowd/DetourLocalBoundary.h
@@ -56,7 +56,7 @@ public:
 	
 	void reset();
 
-	// [UE4: new sections: link removal, path corridor, direction]
+	// [UE: new sections: link removal, path corridor, direction]
 	void update(dtPolyRef ref, const float* pos, const float collisionQueryRange,
 		const bool bIgnoreAtEnd, const float* endPos,
 		const dtPolyRef* path, const int npath,
diff --git a/Engine/Source/Runtime/Navmesh/Public/DetourCrowd/DetourObstacleAvoidance.h b/Engine/Source/Runtime/Navmesh/Public/DetourCrowd/DetourObstacleAvoidance.h
--- a/Engine/Source/Runtime/Navmesh/Public/DetourCrowd/DetourObstacleAvoidance.h
+++ b/Engine/Source/Runtime/Navmesh/Public/DetourCrowd/DetourObstacleAvoidance.h
@@ -90,13 +90,13 @@ struct NAVMESH_API dtObstacleAvoidanceParams
 	float weightSide;
 	float weightToi;
 	float horizTime;
-	unsigned char patternIdx;	///< [UE4] index of custom sampling pattern or 0xff for adaptive
+	unsigned char patternIdx;	///< [UE] index of custom sampling pattern or 0xff for adaptive
 	unsigned char adaptiveDivs;	///< adaptive
 	unsigned char adaptiveRings;	///< adaptive
 	unsigned char adaptiveDepth;	///< adaptive
 };
 
-// [UE4] custom sampling patterns
+// [UE] custom sampling patterns
 struct dtObstacleAvoidancePattern
 {
 	float angles[DT_MAX_CUSTOM_SAMPLES];	///< sample's angle (radians) from desired velocity direction
@@ -119,13 +119,13 @@ public:
 				   
 	void addSegment(const float* p, const float* q, int flags = 0);
 
-	// [UE4] store new sampling pattern
+	// [UE] store new sampling pattern
 	bool setCustomSamplingPattern(int idx, const float* angles, const float* radii, int nsamples);
 
-	// [UE4] get custom sampling pattern
+	// [UE] get custom sampling pattern
 	bool getCustomSamplingPattern(int idx, float* angles, float* radii, int* nsamples);
 
-	// [UE4] sample velocity using custom patterns
+	// [UE] sample velocity using custom patterns
 	int sampleVelocityCustom(const float* pos, const float rad,
 					 		 const float vmax, const float vmult,
 							 const float* vel, const float* dvel, float* nvel,
@@ -138,7 +138,7 @@ public:
 							   const dtObstacleAvoidanceParams* params, 
 							   dtObstacleAvoidanceDebugData* debug = 0);
 	
-	// [UE4] main sampling function
+	// [UE] main sampling function
 	inline int sampleVelocity(const float* pos, const float rad,
 		const float vmax, const float vmult,
 		const float* vel, const float* dvel, float* nvel,
@@ -156,7 +156,7 @@ public:
 	inline int getObstacleSegmentCount() const { return m_nsegments; }
 	const dtObstacleSegment* getObstacleSegment(const int i) { return &m_segments[i]; }
 
-	// [UE4] sampling pattern count accessors
+	// [UE] sampling pattern count accessors
 	inline int getCustomPatternCount() const { return m_maxPatterns; }
 
 private:
diff --git a/Engine/Source/Runtime/Navmesh/Public/DetourCrowd/DetourPathCorridor.h b/Engine/Source/Runtime/Navmesh/Public/DetourCrowd/DetourPathCorridor.h
--- a/Engine/Source/Runtime/Navmesh/Public/DetourCrowd/DetourPathCorridor.h
+++ b/Engine/Source/Runtime/Navmesh/Public/DetourCrowd/DetourPathCorridor.h
@@ -69,9 +69,9 @@ public:
 	///  @param[in]		maxCorners		The maximum number of corners the buffers can hold.
 	///  @param[in]		navquery		The query object used to build the corridor.
 	///  @param[in]		filter			The filter to apply to the operation.
-	///  @param[in]		pathOffsetDistance	[UE4] Radius for path offsetting
-	///  @param[in]		earlyReachDistance	[UE4] Radius for early reach detection
-	///  @param[in]		bAllowEarlyReach [UE4] Check if corner skipping for EarlyReachTest is available now
+	///  @param[in]		pathOffsetDistance	[UE] Radius for path offsetting
+	///  @param[in]		earlyReachDistance	[UE] Radius for early reach detection
+	///  @param[in]		bAllowEarlyReach [UE] Check if corner skipping for EarlyReachTest is available now
 	/// @return The number of corners returned in the corner buffers. [0 <= value <= @p maxCorners]
 	int findCorners(float* cornerVerts, unsigned char* cornerFlags,
 					dtPolyRef* cornerPolys, const int maxCorners,
@@ -96,12 +96,12 @@ public:
 								   float* startPos, float* endPos,
 								   dtNavMeshQuery* navquery);
 
-	/// [UE4] check if offmesh connection can be traversed, but don't modify corridor yet
+	/// [UE] check if offmesh connection can be traversed, but don't modify corridor yet
 	bool canMoveOverOffmeshConnection(dtPolyRef offMeshConRef, dtPolyRef* refs,
 		const float* agentPos, float* startPos, float* endPos,
 		dtNavMeshQuery* navquery) const;
 
-	/// [UE4] remove offmesh connection from corridor
+	/// [UE] remove offmesh connection from corridor
 	void pruneOffmeshConenction(dtPolyRef offMeshConRef);
 
 	bool fixPathStart(dtPolyRef safeRef, const float* safePos);
diff --git a/Engine/Source/Runtime/Navmesh/Public/DetourTileCache/DetourTileCache.h b/Engine/Source/Runtime/Navmesh/Public/DetourTileCache/DetourTileCache.h
--- a/Engine/Source/Runtime/Navmesh/Public/DetourTileCache/DetourTileCache.h
+++ b/Engine/Source/Runtime/Navmesh/Public/DetourTileCache/DetourTileCache.h
@@ -82,14 +82,14 @@ struct dtTileCacheParams
 	float maxSimplificationError;
 	int maxTiles;
 	int maxObstacles;
-//@UE4 BEGIN
+//@UE BEGIN
 	float detailSampleDist;
 	float detailSampleMaxError;
 	int minRegionArea;
 	int mergeRegionArea;
 	int regionChunkSize;
 	int regionPartitioning;
-//@UE4 END
+//@UE END
 };
 
 struct NAVMESH_API dtTileCacheMeshProcess
diff --git a/Engine/Source/Runtime/Navmesh/Public/DetourTileCache/DetourTileCacheBuilder.h b/Engine/Source/Runtime/Navmesh/Public/DetourTileCache/DetourTileCacheBuilder.h
--- a/Engine/Source/Runtime/Navmesh/Public/DetourTileCache/DetourTileCacheBuilder.h
+++ b/Engine/Source/Runtime/Navmesh/Public/DetourTileCache/DetourTileCacheBuilder.h
@@ -80,7 +80,7 @@ struct dtTileCachePolyMesh
 	unsigned short* regs;	///< Region ID of polygon
 };
 
-//@UE4 BEGIN
+//@UE BEGIN
 #if WITH_NAVMESH_CLUSTER_LINKS
 struct dtTileCacheClusterSet
 {
@@ -125,7 +125,7 @@ protected:
 	virtual void doDtLog(const char* /*msg*/, const int /*len*/) {}
 };
 
-//@UE4 END
+//@UE END
 
 struct NAVMESH_API dtTileCacheAlloc
 {
@@ -155,7 +155,7 @@ struct NAVMESH_API dtTileCacheCompressor
 								unsigned char* buffer, const int maxBufferSize, int* bufferSize) = 0;
 };
 
-//@UE4 BEGIN: moved from DetourTileCacheBuilder.cpp
+//@UE BEGIN: moved from DetourTileCacheBuilder.cpp
 template<class T> class dtFixedArray
 {
 	dtTileCacheAlloc* m_alloc;
@@ -183,7 +183,7 @@ inline int getDirOffsetY(int dir)
 	const int offset[4] = { 0, 1, 0, -1 };
 	return offset[dir&0x03];
 }
-//@UE4 END
+//@UE END
 
 NAVMESH_API dtStatus dtBuildTileCacheLayer(dtTileCacheCompressor* comp,
 							   dtTileCacheLayerHeader* header,
@@ -204,7 +204,7 @@ NAVMESH_API void dtFreeTileCacheContourSet(dtTileCacheAlloc* alloc, dtTileCacheC
 NAVMESH_API dtTileCachePolyMesh* dtAllocTileCachePolyMesh(dtTileCacheAlloc* alloc);
 NAVMESH_API void dtFreeTileCachePolyMesh(dtTileCacheAlloc* alloc, dtTileCachePolyMesh* lmesh);
 
-//@UE4 BEGIN
+//@UE BEGIN
 #if WITH_NAVMESH_CLUSTER_LINKS
 NAVMESH_API dtTileCacheClusterSet* dtAllocTileCacheClusterSet(dtTileCacheAlloc* alloc);
 NAVMESH_API void dtFreeTileCacheClusterSet(dtTileCacheAlloc* alloc, dtTileCacheClusterSet* clusters);
@@ -215,12 +215,12 @@ NAVMESH_API void dtFreeTileCachePolyMeshDetail(dtTileCacheAlloc* alloc, dtTileCa
 
 NAVMESH_API dtTileCacheDistanceField* dtAllocTileCacheDistanceField(dtTileCacheAlloc* alloc);
 NAVMESH_API void dtFreeTileCacheDistanceField(dtTileCacheAlloc* alloc, dtTileCacheDistanceField* dfield);
-//@UE4 END
+//@UE END
 
 NAVMESH_API dtStatus dtMarkCylinderArea(dtTileCacheLayer& layer, const float* orig, const float cs, const float ch,
 	const float* pos, const float radius, const float height, const unsigned char areaId);
 
-//@UE4 BEGIN: more shapes
+//@UE BEGIN: more shapes
 NAVMESH_API dtStatus dtMarkBoxArea(dtTileCacheLayer& layer, const float* orig, const float cs, const float ch,
 	const float* pos, const float* extent, const unsigned char areaId);
 
@@ -240,10 +240,10 @@ NAVMESH_API dtStatus dtReplaceConvexArea(dtTileCacheLayer& layer, const float* o
 
 NAVMESH_API dtStatus dtReplaceArea(dtTileCacheLayer& layer, const unsigned char areaId, const unsigned char filterAreaId);
 
-//@UE4 END
+//@UE END
 
 
-//@UE4 BEGIN: renamed building regions to dtBuildTileCacheRegionsMonotone, added new region generation
+//@UE BEGIN: renamed building regions to dtBuildTileCacheRegionsMonotone, added new region generation
 NAVMESH_API dtStatus dtBuildTileCacheDistanceField(dtTileCacheAlloc* alloc, dtTileCacheLayer& layer, dtTileCacheDistanceField& dfield);
 
 NAVMESH_API dtStatus dtBuildTileCacheRegions(dtTileCacheAlloc* alloc,
@@ -257,18 +257,18 @@ NAVMESH_API dtStatus dtBuildTileCacheRegionsMonotone(dtTileCacheAlloc* alloc,
 NAVMESH_API dtStatus dtBuildTileCacheRegionsChunky(dtTileCacheAlloc* alloc,
 								const int minRegionArea, const int mergeRegionArea,
 								dtTileCacheLayer& layer, int regionChunkSize);
-//@UE4 END
+//@UE END
 
 NAVMESH_API dtStatus dtBuildTileCacheContours(dtTileCacheAlloc* alloc,
 								dtTileCacheLayer& layer,
 								const int walkableClimb, const float maxError,
 								const float cs, const float ch,
 								dtTileCacheContourSet& lcset
-								//@UE4 BEGIN
+								//@UE BEGIN
 #if WITH_NAVMESH_CLUSTER_LINKS
 								, dtTileCacheClusterSet& lclusters
 #endif // WITH_NAVMESH_CLUSTER_LINKS
-								//@UE4 END
+								//@UE END
 								);
 
 NAVMESH_API dtStatus dtBuildTileCachePolyMesh(dtTileCacheAlloc* alloc,
@@ -276,7 +276,7 @@ NAVMESH_API dtStatus dtBuildTileCachePolyMesh(dtTileCacheAlloc* alloc,
 								dtTileCacheContourSet& lcset,
 								dtTileCachePolyMesh& mesh);
 
-//@UE4 BEGIN
+//@UE BEGIN
 NAVMESH_API dtStatus dtBuildTileCachePolyMeshDetail(dtTileCacheAlloc* alloc,
 								const float cs, const float ch,
 								const float sampleDist, const float sampleMaxError,
@@ -289,11 +289,11 @@ NAVMESH_API dtStatus dtBuildTileCacheClusters(dtTileCacheAlloc* alloc,
 								dtTileCacheClusterSet& lclusters,
 								dtTileCachePolyMesh& lmesh);
 #endif // WITH_NAVMESH_CLUSTER_LINKS
-//@UE4 END
+//@UE END
 
 /// Swaps the endianess of the compressed tile data's header (#dtTileCacheLayerHeader).
 /// Tile layer data does not need endian swapping as it consist only of bytes.
-/// UE4: not anymore, there are short types as well now
+/// UE: not anymore, there are short types as well now
 ///  @param[in,out]	data		The tile data array.
 ///  @param[in]		dataSize	The size of the data array.
 NAVMESH_API bool dtTileCacheHeaderSwapEndian(unsigned char* data, const int dataSize);
diff --git a/Engine/Source/Runtime/Navmesh/Public/Recast/Recast.h b/Engine/Source/Runtime/Navmesh/Public/Recast/Recast.h
--- a/Engine/Source/Runtime/Navmesh/Public/Recast/Recast.h
+++ b/Engine/Source/Runtime/Navmesh/Public/Recast/Recast.h
@@ -441,7 +441,7 @@ struct rcContourSet
 	int borderSize;		///< The AABB border size used to generate the source data from which the contours were derived.
 };
 
-// @UE4 BEGIN
+// @UE BEGIN
 #if WITH_NAVMESH_CLUSTER_LINKS
 /// Represents group of clusters
 /// @ingroup recast
@@ -453,7 +453,7 @@ struct rcClusterSet
 	unsigned short* links;	///< Neighbor Ids per cluster [Size: sum of #nlinks]
 };
 #endif // WITH_NAVMESH_CLUSTER_LINKS
-// @UE4 END 
+// @UE END 
  
 /// Represents a polygon mesh suitable for use in building a navigation mesh. 
 /// @ingroup recast
@@ -541,7 +541,7 @@ NAVMESH_API rcContourSet* rcAllocContourSet();
 ///  @see rcAllocContourSet
 NAVMESH_API void rcFreeContourSet(rcContourSet* cset);
 
-// @UE4 BEGIN
+// @UE BEGIN
 #if WITH_NAVMESH_CLUSTER_LINKS
 /// Allocates a cluster set object using the Recast allocator.
 ///  @return A cluster set that is ready for initialization, or null on failure.
@@ -555,7 +555,7 @@ NAVMESH_API rcClusterSet* rcAllocClusterSet();
 ///  @see rcAllocClusterSet
 NAVMESH_API void rcFreeClusterSet(rcClusterSet* clset);
 #endif // WITH_NAVMESH_CLUSTER_LINKS
-// @UE4 END 
+// @UE END 
 
 /// Allocates a polygon mesh object using the Recast allocator.
 ///  @return A polygon mesh that is ready for initialization, or null on failure.
@@ -614,14 +614,14 @@ enum rcBuildContoursFlags
 	RC_CONTOUR_TESS_AREA_EDGES = 0x02,	///< Tessellate edges between areas during contour simplification.
 };
 
-// UE4
+// UE
 enum rcFilterLowAreaFlags
 {
 	RC_LOW_FILTER_SEED_SPANS = 0x01,	///< initial seeding on spans
 	RC_LOW_FILTER_POST_PROCESS = 0x02,	///< additional filtering at the end
 };
 
-// UE4
+// UE
 enum rcRasterizationFlags
 {
 	RC_PROJECT_TO_BOTTOM = 1 << 0,		///< Will create spans from the triangle surface to the bottom of the heightfield
@@ -933,12 +933,12 @@ NAVMESH_API void rcCacheSpans(rcContext* ctx, rcHeightfield& hf, rcSpanCache* ca
 ///  @param[in,out]	solid			An initialized heightfield.
 ///  @param[in]		flagMergeThr	The distance where the walkable flag is favored over the non-walkable flag.
 ///  								[Limit: >= 0] [Units: vx]
-///  @param[in]     rtzFlags		Flags to change the rasterization behavior			//UE4
-///  @param[in]     rtzMasks		Mask for the rasterization flags [Size: hf.w*hf.h]	//UE4
+///  @param[in]     rtzFlags		Flags to change the rasterization behavior			//UE
+///  @param[in]     rtzMasks		Mask for the rasterization flags [Size: hf.w*hf.h]	//UE
 NAVMESH_API void rcRasterizeTriangle(rcContext* ctx, const float* v0, const float* v1, const float* v2,
 						 const unsigned char area, rcHeightfield& solid,
 						 const int flagMergeThr = 1, 
-						 const int rasterizationFlags = 0, const int* rasterizationMasks = nullptr); //UE4
+						 const int rasterizationFlags = 0, const int* rasterizationMasks = nullptr); //UE
 
 /// Rasterizes an indexed triangle mesh into the specified heightfield.
 ///  @ingroup recast
@@ -951,12 +951,12 @@ NAVMESH_API void rcRasterizeTriangle(rcContext* ctx, const float* v0, const floa
 ///  @param[in,out]	solid			An initialized heightfield.
 ///  @param[in]		flagMergeThr	The distance where the walkable flag is favored over the non-walkable flag. 
 ///  								[Limit: >= 0] [Units: vx]
-///  @param[in]     rtzFlags		Flags to change the rasterization behavior			//UE4
-///  @param[in]     rtzMasks		Mask for the rasterization flags [Size: hf.w*hf.h]	//UE4
+///  @param[in]     rtzFlags		Flags to change the rasterization behavior			//UE
+///  @param[in]     rtzMasks		Mask for the rasterization flags [Size: hf.w*hf.h]	//UE
 NAVMESH_API void rcRasterizeTriangles(rcContext* ctx, const float* verts, const int nv,
 						  const int* tris, const unsigned char* areas, const int nt,
 						  rcHeightfield& solid, const int flagMergeThr = 1, 
-						  const int rasterizationFlags = 0, const int* rasterizationMasks = nullptr); //UE4
+						  const int rasterizationFlags = 0, const int* rasterizationMasks = nullptr); //UE
 
 /// Rasterizes an indexed triangle mesh into the specified heightfield.
 ///  @ingroup recast
@@ -969,12 +969,12 @@ NAVMESH_API void rcRasterizeTriangles(rcContext* ctx, const float* verts, const
 ///  @param[in,out]	solid			An initialized heightfield.
 ///  @param[in]		flagMergeThr	The distance where the walkable flag is favored over the non-walkable flag. 
 ///  								[Limit: >= 0] [Units: vx]
-///  @param[in]     rtzFlags		Flags to change the rasterization behavior			//UE4
-///  @param[in]     rtzMasks		Mask for the rasterization flags [Size: hf.w*hf.h]	//UE4
+///  @param[in]     rtzFlags		Flags to change the rasterization behavior			//UE
+///  @param[in]     rtzMasks		Mask for the rasterization flags [Size: hf.w*hf.h]	//UE
 NAVMESH_API void rcRasterizeTriangles(rcContext* ctx, const float* verts, const int nv,
 						  const unsigned short* tris, const unsigned char* areas, const int nt,
 						  rcHeightfield& solid, const int flagMergeThr = 1, 
-						  const int rasterizationFlags = 0, const int* rasterizationMasks = nullptr); //UE4
+						  const int rasterizationFlags = 0, const int* rasterizationMasks = nullptr); //UE
 
 /// Rasterizes triangles into the specified heightfield.
 ///  @ingroup recast
@@ -985,11 +985,11 @@ NAVMESH_API void rcRasterizeTriangles(rcContext* ctx, const float* verts, const
 ///  @param[in,out]	solid			An initialized heightfield.
 ///  @param[in]		flagMergeThr	The distance where the walkable flag is favored over the non-walkable flag. 
 ///  								[Limit: >= 0] [Units: vx]
-///  @param[in]     rtzFlags		Flags to change the rasterization behavior			//UE4
-///  @param[in]     rtzMasks		Mask for the rasterization flags [Size: hf.w*hf.h]	//UE4
+///  @param[in]     rtzFlags		Flags to change the rasterization behavior			//UE
+///  @param[in]     rtzMasks		Mask for the rasterization flags [Size: hf.w*hf.h]	//UE
 NAVMESH_API void rcRasterizeTriangles(rcContext* ctx, const float* verts, const unsigned char* areas, const int nt,
 						  rcHeightfield& solid, const int flagMergeThr = 1,
-						  const int rasterizationFlags = 0, const int* rasterizationMasks = nullptr); //UE4
+						  const int rasterizationFlags = 0, const int* rasterizationMasks = nullptr); //UE
 
 /// Marks non-walkable spans as walkable if their maximum is within @p walkableClimp of a walkable neighbor. 
 ///  @ingroup recast
@@ -1261,7 +1261,7 @@ inline int rcGetDirOffsetY(int dir)
 /// @see rcHeightfieldLayer, rcContourSet, rcPolyMesh, rcPolyMeshDetail
 /// @{
 
-// @UE4 BEGIN: renamed building layers to rcBuildHeightfieldLayersMonotone, added flood fill based implementation 
+// @UE BEGIN: renamed building layers to rcBuildHeightfieldLayersMonotone, added flood fill based implementation 
 
 /// Builds a layer set from the specified compact heightfield.
 ///  @ingroup recast
@@ -1307,7 +1307,7 @@ NAVMESH_API bool rcBuildHeightfieldLayersChunky(rcContext* ctx, rcCompactHeightf
 									const int chunkSize,
 									rcHeightfieldLayerSet& lset);
 
-// @UE4 END
+// @UE END
 
 /// Builds a contour set from the region outlines in the provided compact heightfield.
 ///  @ingroup recast
@@ -1324,7 +1324,7 @@ NAVMESH_API bool rcBuildContours(rcContext* ctx, rcCompactHeightfield& chf,
 					 const float maxError, const int maxEdgeLen,
 					 rcContourSet& cset, const int flags = RC_CONTOUR_TESS_WALL_EDGES);
 
-//@UE4 BEGIN
+//@UE BEGIN
 #if WITH_NAVMESH_CLUSTER_LINKS
 /// Builds a cluster set from contours
 /// @ingroup recast
@@ -1333,7 +1333,7 @@ NAVMESH_API bool rcBuildContours(rcContext* ctx, rcCompactHeightfield& chf,
 /// @param[out]		clusters	Resulting cluster set
 NAVMESH_API bool rcBuildClusters(rcContext* ctx, rcContourSet& cset, rcClusterSet& clusters);
 #endif // WITH_NAVMESH_CLUSTER_LINKS
-//@UE4 END
+//@UE END
 
 /// Builds a polygon mesh from the provided contours.
 ///  @ingroup recast
diff --git a/Engine/Source/Runtime/Navmesh/Public/Recast/RecastAlloc.h b/Engine/Source/Runtime/Navmesh/Public/Recast/RecastAlloc.h
--- a/Engine/Source/Runtime/Navmesh/Public/Recast/RecastAlloc.h
+++ b/Engine/Source/Runtime/Navmesh/Public/Recast/RecastAlloc.h
@@ -105,9 +105,9 @@ public:
 	/// The current size of the integer array.
 	inline int size() const { return m_size; }
 
-	// @UE4 BEGIN: added contains()
+	// @UE BEGIN: added contains()
 	bool contains(int n) const;
-	// @UE4 END
+	// @UE END
 };
 
 /// A simple helper class used to delete an array when it goes out of scope.
@@ -132,7 +132,7 @@ public:
 	///  @return The root array pointer.
 	inline operator T*() { return ptr; }
 
-	/// UE4: resize and copy existing memory (n = element count), doesn't destruct elements!
+	/// UE: resize and copy existing memory (n = element count), doesn't destruct elements!
 	bool resizeGrow(int n);
 };
 
