[Navigation] Fix issue in dtCrowd::updateStepProximityData(), cornerVerts data was wrongly accessed.

diff --git a/Engine/Source/Runtime/Navmesh/Private/DetourCrowd/DetourCrowd.cpp b/Engine/Source/Runtime/Navmesh/Private/DetourCrowd/DetourCrowd.cpp
--- a/Engine/Source/Runtime/Navmesh/Private/DetourCrowd/DetourCrowd.cpp
+++ b/Engine/Source/Runtime/Navmesh/Private/DetourCrowd/DetourCrowd.cpp
@@ -1288,7 +1288,7 @@ void dtCrowd::updateStepProximityData(const float dt, dtCrowdAgentDebugInfo* deb
 			float moveDir[3] = { 0.0f };
 			if (ag->ncorners)
 			{
-				dtVsub(moveDir, &ag->cornerVerts[2], &ag->cornerVerts[0]);
+				dtVsub(moveDir, &ag->cornerVerts[0], ag->npos);
 			}
 			else
 			{
