[Navmesh] Allow unconnected single cells to build contours

diff --git a/Engine/Source/Runtime/Navmesh/Private/DetourTileCache/DetourTileCacheBuilder.cpp b/Engine/Source/Runtime/Navmesh/Private/DetourTileCache/DetourTileCacheBuilder.cpp
--- a/Engine/Source/Runtime/Navmesh/Private/DetourTileCache/DetourTileCacheBuilder.cpp
+++ b/Engine/Source/Runtime/Navmesh/Private/DetourTileCache/DetourTileCacheBuilder.cpp
@@ -801,9 +801,8 @@ dtStatus dtBuildTileCacheContours(dtTileCacheAlloc* alloc, dtTileCacheLayer& lay
 		for (int x = 0; x < w; ++x)
 		{
 			const int idx = x+y*w;
-			if (flags[idx] == 0 || flags[idx] == 0xf)
+			if (flags[idx] == 0)
 			{
-				flags[idx] = 0;
 				continue;
 			}
 
