[Navmesh] When building using chunky region partitioning, include remaining non dividable cells into the last split.
This prevents having one cell wide region strips and one cell regions and the 'Region Chunk Splits' setting is now correctly respected.
Before this change, a region split on 2 could result into 3 splits (ex: a 17 cells wide tile would have been divided into 8-8-1 splits, with this change it will be 8-9).

diff --git a/Engine/Source/Runtime/Navmesh/Private/DetourTileCache/DetourTileCacheRegion.cpp b/Engine/Source/Runtime/Navmesh/Private/DetourTileCache/DetourTileCacheRegion.cpp
--- a/Engine/Source/Runtime/Navmesh/Private/DetourTileCache/DetourTileCacheRegion.cpp
+++ b/Engine/Source/Runtime/Navmesh/Private/DetourTileCache/DetourTileCacheRegion.cpp
@@ -1226,28 +1226,39 @@ static dtStatus CollectRegionsChunky(dtTileCacheAlloc* alloc, dtTileCacheLayer&
 	dtIntArray prevCount(256);
 	unsigned short regId = 0;
 
-	for (int chunkx = 0; chunkx < w; chunkx += chunkSize)
+	const int numXChunks = dtMax(1, w / chunkSize); // Division floors, remainder is put to the last chunk.
+	const int numYChunks = dtMax(1, h / chunkSize);
+
+	for (int chunkx = 0; chunkx < numXChunks; chunkx++)
 	{
-		for (int chunky = 0; chunky < h; chunky += chunkSize)
+		const bool lastXChunk = chunkx == (numXChunks - 1);
+		const int minx = chunkx * chunkSize;
+		const int maxx = lastXChunk ? w : (minx + chunkSize);
+
+		for (int chunky = 0; chunky < numYChunks; chunky++)
 		{
-			const int maxx = dtMin(chunkx + chunkSize, w);
-			const int maxy = dtMin(chunky + chunkSize, h);
+			const bool lastYChunk = chunky == (numYChunks - 1);
+			const int miny = chunky * chunkSize;
+			const int maxy = lastYChunk ? h : (miny + chunkSize);
 
-			for (int y = chunky; y < maxy; ++y)
+			for (int y = miny; y < maxy; ++y)
 			{
 				prevCount.resize(regId+1);
 				memset(&prevCount[0],0,sizeof(int)*regId);
 				unsigned short sweepId = 0;
 
-				for (int x = chunkx; x < maxx; ++x)
+				for (int x = minx; x < maxx; ++x)
 				{
 					const int idx = x + y*w;
-					if (layer.areas[idx] == DT_TILECACHE_NULL_AREA) continue;
+					if (layer.areas[idx] == DT_TILECACHE_NULL_AREA)
+					{
+						continue;
+					}
 
 					unsigned short sid = 0xffff;
 
 					// -x
-					if (x > chunkx && isConnected(layer, idx, 0))
+					if (x > minx && isConnected(layer, idx, 0))
 					{
 						const int xidx = (x-1)+y*w;
 						if (layer.regs[xidx] != 0xffff && layer.areas[xidx] == layer.areas[idx])
@@ -1262,7 +1273,7 @@ static dtStatus CollectRegionsChunky(dtTileCacheAlloc* alloc, dtTileCacheLayer&
 					}
 
 					// -y
-					if (y > chunky && isConnected(layer, idx, 3))
+					if (y > miny && isConnected(layer, idx, 3))
 					{
 						const int yidx = x+(y-1)*w;
 						const unsigned short nr = layer.regs[yidx];
@@ -1306,7 +1317,7 @@ static dtStatus CollectRegionsChunky(dtTileCacheAlloc* alloc, dtTileCacheLayer&
 				}
 
 				// Remap local sweep ids to region ids.
-				for (int x = chunkx; x < maxx; ++x)
+				for (int x = minx; x < maxx; ++x)
 				{
 					const int idx = x+y*w;
 					if (layer.regs[idx] != 0xffff)
@@ -1348,7 +1359,8 @@ static dtStatus CollectRegionsChunky(dtTileCacheAlloc* alloc, dtTileCacheLayer&
 	// Find region neighbours.
 	for (int y = 0; y < h; ++y)
 	{
-		const int chunkYOffset = (y / chunkSize) * chunkSize;
+		const int ciy = dtMin(y / chunkSize , numYChunks-1);
+
 		const bool borderY = (y == 0) || (y == (h - 1));
 		for (int x = 0; x < w; ++x)
 		{
@@ -1360,7 +1372,9 @@ static dtStatus CollectRegionsChunky(dtTileCacheAlloc* alloc, dtTileCacheLayer&
 			// Update area.
 			regs[ri].area++;
 			regs[ri].areaId = layer.areas[idx];
-			regs[ri].chunkId = (x / chunkSize) + chunkYOffset;
+
+			const int cix = dtMin(x / chunkSize , numXChunks-1);
+			regs[ri].chunkId = cix + (ciy * numXChunks);
 			regs[ri].border |= borderY || (x == 0) || (x == (w - 1));
 
 			// Update neighbours
