Memory optimizations for Recast (Detour). I used UE_LOG in recast as it's already there. Also I only added [REDACTED_HANDLE] BEGIN [REDACTED_HANDLE] END for the big stuff as a hint that its been changed in case of issues.

diff --git a/Engine/Source/Runtime/NavigationSystem/Private/NavMesh/PImplRecastNavMesh.cpp b/Engine/Source/Runtime/NavigationSystem/Private/NavMesh/PImplRecastNavMesh.cpp
--- a/Engine/Source/Runtime/NavigationSystem/Private/NavMesh/PImplRecastNavMesh.cpp
+++ b/Engine/Source/Runtime/NavigationSystem/Private/NavMesh/PImplRecastNavMesh.cpp
@@ -859,14 +859,12 @@ void FPImplRecastNavMesh::SerializeRecastMeshTile(FArchive& ArWrapped, int32 Nav
 		// now serialize the data in the blob!
 
 		// header
-		Ar << Header->magic << Header->version << Header->x << Header->y;
-		Ar << Header->layer << Header->userId << Header->polyCount << Header->vertCount;
+		Ar << Header->version << Header->x << Header->y;
+		Ar << Header->layer << Header->polyCount << Header->vertCount;
 		Ar << Header->maxLinkCount << Header->detailMeshCount << Header->detailVertCount << Header->detailTriCount;
 		Ar << Header->bvNodeCount << Header->offMeshConCount<< Header->offMeshBase;
-		Ar << Header->walkableHeight << Header->walkableRadius << Header->walkableClimb;
 		Ar << Header->bmin[0] << Header->bmin[1] << Header->bmin[2];
 		Ar << Header->bmax[0] << Header->bmax[1] << Header->bmax[2];
-		Ar << Header->bvQuantFactor;
 #if WITH_NAVMESH_CLUSTER_LINKS
 		Ar << Header->clusterCount;
 #else
@@ -1077,7 +1075,6 @@ void FPImplRecastNavMesh::SerializeCompressedTileCacheData(FArchive& ArWrapped,
 
 	// Serialize dtTileCacheLayerHeader by hand so we can account for the FReals always being serialized as doubles
 	dtTileCacheLayerHeader* Header = (dtTileCacheLayerHeader*)CompressedData;
-	Ar << Header->magic;
 	Ar << Header->version;
 	Ar << Header->tx;
 	Ar << Header->ty;
@@ -3346,4 +3343,4 @@ bool FPImplRecastNavMesh::HasTileCacheLayers(int32 TileX, int32 TileY) const
 
 #undef INITIALIZE_NAVQUERY
 
-#endif // WITH_RECAST
+#endif // WITH_RECAST
\ No newline at end of file
diff --git a/Engine/Source/Runtime/NavigationSystem/Private/NavMesh/RecastNavMesh.cpp b/Engine/Source/Runtime/NavigationSystem/Private/NavMesh/RecastNavMesh.cpp
--- a/Engine/Source/Runtime/NavigationSystem/Private/NavMesh/RecastNavMesh.cpp
+++ b/Engine/Source/Runtime/NavigationSystem/Private/NavMesh/RecastNavMesh.cpp
@@ -68,7 +68,7 @@ FDetourTileLayout::FDetourTileLayout(const dtMeshTile& tile)
 {
 	const dtMeshHeader* header = tile.header;
 
-	if (header && (header->magic == DT_NAVMESH_MAGIC) && (header->version == DT_NAVMESH_VERSION))
+	if (header && (header->version == DT_NAVMESH_VERSION))
 	{
 		FDetourTileSizeInfo SizeInfo;
 
@@ -754,6 +754,23 @@ void ARecastNavMesh::RestrictBuildingToActiveTiles(bool InRestrictBuildingToActi
 	}
 }
 
+void ARecastNavMesh::OnRegistered()
+{
+	Super::OnRegistered();
+
+	check(RecastNavMeshImpl);
+
+	// This check can fail when the NavMeshVersion indicates the map needs the nav mesh rebuilt
+	if (RecastNavMeshImpl->GetRecastMesh())
+	{
+		// Set walkable climb / radius / height in recast remembering to convert to recast cell units
+		RecastNavMeshImpl->GetRecastMesh()->setWalkableClimb(AgentMaxStepHeight);
+		RecastNavMeshImpl->GetRecastMesh()->setWalkableHeight(AgentHeight);
+		RecastNavMeshImpl->GetRecastMesh()->setWalkableRadius(AgentRadius);
+		RecastNavMeshImpl->GetRecastMesh()->setBVQuantFactor(1.f / CellSize);
+	}
+}
+
 void ARecastNavMesh::SerializeRecastNavMesh(FArchive& Ar, FPImplRecastNavMesh*& NavMesh, int32 InNavMeshVersion)
 {
 	if (!Ar.IsLoading()	&& NavMesh == NULL)
diff --git a/Engine/Source/Runtime/NavigationSystem/Private/NavMesh/RecastNavMeshDataChunk.cpp b/Engine/Source/Runtime/NavigationSystem/Private/NavMesh/RecastNavMeshDataChunk.cpp
--- a/Engine/Source/Runtime/NavigationSystem/Private/NavMesh/RecastNavMeshDataChunk.cpp
+++ b/Engine/Source/Runtime/NavigationSystem/Private/NavMesh/RecastNavMeshDataChunk.cpp
@@ -352,7 +352,7 @@ void URecastNavMeshDataChunk::MoveTiles(FPImplRecastNavMesh& NavMeshImpl, const
 				const FVector::FReal TileHeight = NavMesh->getParams()->tileHeight;
 
 				const dtMeshHeader* Header = (dtMeshHeader*)TileData.TileRawData->RawData;
-				if (Header->magic != DT_NAVMESH_MAGIC || Header->version != DT_NAVMESH_VERSION)
+				if (Header->version != DT_NAVMESH_VERSION)
 				{
 					continue;
 				}
@@ -366,7 +366,7 @@ void URecastNavMeshDataChunk::MoveTiles(FPImplRecastNavMesh& NavMeshImpl, const
 
 				const int OffsetWithRotX = Offset.X + DeltaX;
 				const int OffsetWithRotY = Offset.Y + DeltaY;
-				const bool bSuccess = dtTransformTileData(TileData.TileRawData->RawData, TileData.TileDataSize, OffsetWithRotX, OffsetWithRotY, TileWidth, TileHeight, RotationDeg);
+				const bool bSuccess = dtTransformTileData(TileData.TileRawData->RawData, TileData.TileDataSize, OffsetWithRotX, OffsetWithRotY, TileWidth, TileHeight, RotationDeg, NavMesh->getBVQuantFactor());
 				UE_CLOG(bSuccess, LogNavigation, Verbose, TEXT("   Moved tile from (%i,%i) to (%i,%i)."), TileData.OriginalX, TileData.OriginalY, (TileData.OriginalX + OffsetWithRotX), (TileData.OriginalY + OffsetWithRotY));
 			}
 		}
diff --git a/Engine/Source/Runtime/NavigationSystem/Private/NavMesh/RecastNavMeshGenerator.cpp b/Engine/Source/Runtime/NavigationSystem/Private/NavMesh/RecastNavMeshGenerator.cpp
--- a/Engine/Source/Runtime/NavigationSystem/Private/NavMesh/RecastNavMeshGenerator.cpp
+++ b/Engine/Source/Runtime/NavigationSystem/Private/NavMesh/RecastNavMeshGenerator.cpp
@@ -3316,7 +3316,6 @@ bool FRecastTileGenerator::RecastBuildTileCache(FNavMeshBuildContext& BuildConte
 
 		// Store header
 		dtTileCacheLayerHeader header;
-		header.magic = DT_TILECACHE_MAGIC;
 		header.version = DT_TILECACHE_VERSION;
 
 		// Tile layer location in the navmesh.
@@ -4723,6 +4722,16 @@ void FRecastNavMeshGenerator::Init()
 			Config.MaxPolysPerTile = SavedNavParams->maxPolys;
 		}
 	}
+
+	// This check can fail when the NavMeshVersion indicates the map needs the nav mesh rebuilt
+	if (DestNavMesh->GetRecastMesh())
+	{
+		DestNavMesh->GetRecastMesh()->setWalkableClimb(Config.AgentMaxClimb);
+		DestNavMesh->GetRecastMesh()->setWalkableHeight(Config.AgentHeight);
+		DestNavMesh->GetRecastMesh()->setWalkableRadius(Config.AgentRadius);
+		DestNavMesh->GetRecastMesh()->setBVQuantFactor(1.f / Config.cs);
+	}
+
 	UpdateNavigationBounds();
 
 	/** setup maximum number of active tile generator*/
diff --git a/Engine/Source/Runtime/NavigationSystem/Public/NavMesh/RecastNavMesh.h b/Engine/Source/Runtime/NavigationSystem/Public/NavMesh/RecastNavMesh.h
--- a/Engine/Source/Runtime/NavigationSystem/Public/NavMesh/RecastNavMesh.h
+++ b/Engine/Source/Runtime/NavigationSystem/Public/NavMesh/RecastNavMesh.h
@@ -893,6 +893,8 @@ protected:
 	TArray<FIntPoint>& GetActiveTiles(); 
 	virtual void RestrictBuildingToActiveTiles(bool InRestrictBuildingToActiveTiles) override;
 
+	virtual void OnRegistered() override;
+
 public:
 	/** Whether NavMesh should adjust his tile pool size when NavBounds are changed */
 	bool IsResizable() const;
diff --git a/Engine/Source/Runtime/NavigationSystem/Public/NavMesh/RecastVersion.h b/Engine/Source/Runtime/NavigationSystem/Public/NavMesh/RecastVersion.h
--- a/Engine/Source/Runtime/NavigationSystem/Public/NavMesh/RecastVersion.h
+++ b/Engine/Source/Runtime/NavigationSystem/Public/NavMesh/RecastVersion.h
@@ -17,7 +17,8 @@
 #define NAVMESHVER_LWCOORDS					14
 #define NAVMESHVER_OODLE_COMPRESSION		15
 #define NAVMESHVER_LWCOORDS_SEREALIZATION 	17 // Allows for nav meshes to be serialized agnostic of LWCoords being float or double.
+#define NAVMESHVER_LWCOORDS_OPTIMIZATION	18
 
 
-#define NAVMESHVER_LATEST				NAVMESHVER_LWCOORDS_SEREALIZATION
-#define NAVMESHVER_MIN_COMPATIBLE		NAVMESHVER_LWCOORDS_SEREALIZATION
+#define NAVMESHVER_LATEST				NAVMESHVER_LWCOORDS_OPTIMIZATION
+#define NAVMESHVER_MIN_COMPATIBLE		NAVMESHVER_LWCOORDS_OPTIMIZATION
diff --git a/Engine/Source/Runtime/Navmesh/Private/DebugUtils/DetourDebugDraw.cpp b/Engine/Source/Runtime/Navmesh/Private/DebugUtils/DetourDebugDraw.cpp
--- a/Engine/Source/Runtime/Navmesh/Private/DebugUtils/DetourDebugDraw.cpp
+++ b/Engine/Source/Runtime/Navmesh/Private/DebugUtils/DetourDebugDraw.cpp
@@ -364,10 +364,10 @@ void duDebugDrawNavMeshNodes(struct duDebugDraw* dd, const dtNavMeshQuery& query
 }
 
 
-static void drawMeshTileBVTree(duDebugDraw* dd, const dtMeshTile* tile)
+static void drawMeshTileBVTree(duDebugDraw* dd, const dtMeshTile* tile, const dtReal bvQuantFactor)
 {
 	// Draw BV nodes.
-	const dtReal cs = 1.0f / tile->header->bvQuantFactor;
+	const dtReal cs = 1.0f / bvQuantFactor;
 	dd->begin(DU_DRAW_LINES, 1.0f);
 	for (int i = 0; i < tile->header->bvNodeCount; ++i)
 	{
@@ -388,20 +388,22 @@ static void drawMeshTileBVTree(duDebugDraw* dd, const dtMeshTile* tile)
 void duDebugDrawNavMeshBVTree(duDebugDraw* dd, const dtNavMesh& mesh)
 {
 	if (!dd) return;
+
+	const dtReal bvQuantFactor = mesh.getBVQuantFactor();
 	
 	for (int i = 0; i < mesh.getMaxTiles(); ++i)
 	{
 		const dtMeshTile* tile = mesh.getTile(i);
 		if (!tile->header) continue;
-		drawMeshTileBVTree(dd, tile);
+		drawMeshTileBVTree(dd, tile, bvQuantFactor);
 	}
 }
 
-static void drawMeshTilePortal(duDebugDraw* dd, const dtMeshTile* tile)
+static void drawMeshTilePortal(duDebugDraw* dd, const dtMeshTile* tile, const dtNavMesh& mesh)
 {
 	// Draw portals
 	const dtReal padx = 0.04f;
-	const dtReal pady = tile->header->walkableClimb;
+	const dtReal pady = mesh.getWalkableClimb();
 
 	dd->begin(DU_DRAW_LINES, 2.0f);
 
@@ -477,7 +479,7 @@ void duDebugDrawNavMeshPortals(duDebugDraw* dd, const dtNavMesh& mesh)
 	{
 		const dtMeshTile* tile = mesh.getTile(i);
 		if (!tile->header) continue;
-		drawMeshTilePortal(dd, tile);
+		drawMeshTilePortal(dd, tile, mesh);
 	}
 }
 
diff --git a/Engine/Source/Runtime/Navmesh/Private/Detour/DetourNavMesh.cpp b/Engine/Source/Runtime/Navmesh/Private/Detour/DetourNavMesh.cpp
--- a/Engine/Source/Runtime/Navmesh/Private/Detour/DetourNavMesh.cpp
+++ b/Engine/Source/Runtime/Navmesh/Private/Detour/DetourNavMesh.cpp
@@ -23,6 +23,8 @@
 #include "Detour/DetourCommon.h"
 #include "Detour/DetourAssert.h"
 
+DEFINE_LOG_CATEGORY(LogDetour);
+
 //@UE BEGIN Adding support for memory tracking.
 static dtStatsPostAddTileFunc* sAddTileFunc = nullptr;
 static dtStatsPreRemoveTileFunc* sRemoveTileFunc = nullptr;
@@ -420,11 +422,11 @@ int segmentIntersectionSorter(const void* i1, const void* i2)
 }
 
 static void gatherSegmentIntersections(dtMeshTile* tile,
-	const dtReal* spos, const dtReal* epos, const dtReal radius,
+	const dtReal* spos, const dtReal* epos,const dtReal radius, const dtReal walkableClimb,
 	dtOffMeshSegmentTileIntersection& list)
 {
 	// get all polys intersecting with segment
-	dtReal segBMin[3], segBMax[3], segRad[3] = { radius, tile->header->walkableClimb, radius };
+	dtReal segBMin[3], segBMax[3], segRad[3] = { radius, walkableClimb, radius };
 	dtVcopy(segBMin, spos);
 	dtVcopy(segBMax, spos);
 	dtVmin(segBMin, epos);
@@ -468,7 +470,7 @@ static void gatherSegmentIntersections(dtMeshTile* tile,
 	}
 }
 
-static dtOffMeshSegmentData* initSegmentIntersection(dtMeshTile* tile)
+static dtOffMeshSegmentData* initSegmentIntersection(const dtNavMesh& navMesh, dtMeshTile* tile)
 {
 	const int segCount = tile->header->offMeshSegConCount;
 	if (segCount <= 0)
@@ -484,14 +486,14 @@ static dtOffMeshSegmentData* initSegmentIntersection(dtMeshTile* tile)
 		dtOffMeshSegmentConnection& con = tile->offMeshSeg[i];
 		
 		CA_SUPPRESS(6385);
-		gatherSegmentIntersections(tile, con.startA, con.endA, con.rad, segs[i].listA);
-		gatherSegmentIntersections(tile, con.startB, con.endB, con.rad, segs[i].listB);
+		gatherSegmentIntersections(tile, con.startA, con.endA, con.rad, navMesh.getWalkableClimb(), segs[i].listA);
+		gatherSegmentIntersections(tile, con.startB, con.endB, con.rad, navMesh.getWalkableClimb(), segs[i].listB);
 	}
 
 	return segs;
 }
 
-static void appendSegmentIntersection(dtOffMeshSegmentData* seg, dtMeshTile* tile, dtMeshTile* nei)
+static void appendSegmentIntersection(const dtNavMesh& navMesh, dtOffMeshSegmentData* seg, dtMeshTile* tile, dtMeshTile* nei)
 {
 	if (seg == NULL)
 		return;
@@ -500,8 +502,8 @@ static void appendSegmentIntersection(dtOffMeshSegmentData* seg, dtMeshTile* til
 	{
 		dtOffMeshSegmentConnection& con = tile->offMeshSeg[i];
 
-		gatherSegmentIntersections(nei, con.startA, con.endA, con.rad, seg[i].listA);
-		gatherSegmentIntersections(nei, con.startB, con.endB, con.rad, seg[i].listB);
+		gatherSegmentIntersections(nei, con.startA, con.endA, con.rad, navMesh.getWalkableClimb(), seg[i].listA);
+		gatherSegmentIntersections(nei, con.startB, con.endB, con.rad, navMesh.getWalkableClimb(), seg[i].listB);
 	}
 }
 
@@ -778,6 +780,10 @@ Notes:
 dtNavMesh::dtNavMesh() :
 	m_tileWidth(0),
 	m_tileHeight(0),
+	m_walkableHeight(0),
+	m_walkableRadius(0),
+	m_walkableClimb(0),
+	m_bvQuantFactor(0),
 	m_maxTiles(0),
 	m_tileLutSize(0),
 	m_tileLutMask(0),
@@ -865,8 +871,6 @@ dtStatus dtNavMesh::init(unsigned char* data, const int dataSize, const int flag
 {
 	// Make sure the data is in right format.
 	dtMeshHeader* header = (dtMeshHeader*)data;
-	if (header->magic != DT_NAVMESH_MAGIC)
-		return DT_FAILURE | DT_WRONG_MAGIC;
 	if (header->version != DT_NAVMESH_VERSION)
 		return DT_FAILURE | DT_WRONG_VERSION;
 
@@ -934,7 +938,7 @@ int dtNavMesh::findConnectingPolys(const dtReal* va, const dtReal* vb,
 			calcSlabEndPoints(vc, vd, bmin, bmax, side);
 
 			unsigned char overlapMode = 0;
-			if (!overlapSlabs(amin, amax, bmin, bmax, 0.01f, tile->header->walkableClimb, &overlapMode)) continue;
+			if (!overlapSlabs(amin, amax, bmin, bmax, 0.01f, m_walkableClimb, &overlapMode)) continue;
 
 			// if overlapping with only one side, verify height difference using detailed mesh
 			if (overlapMode == SLABOVERLAP_Max || overlapMode == SLABOVERLAP_Min)
@@ -947,7 +951,7 @@ int dtNavMesh::findConnectingPolys(const dtReal* va, const dtReal* vb,
 				const dtReal aH = getHeightFromDMesh(fromTile, fromPolyIdx, apt);
 				const dtReal bH = getHeightFromDMesh(tile, i, bpt);
 				const dtReal heightDiff = dtAbs(aH - bH);
-				if (heightDiff > tile->header->walkableClimb)
+				if (heightDiff > m_walkableClimb)
 					continue;
 			}
 
@@ -1599,7 +1603,6 @@ int dtNavMesh::queryPolygonsInTile(const dtMeshTile* tile, const dtReal* qmin, c
 		const dtBVNode* end = &tile->bvTree[tile->header->bvNodeCount];
 		const dtReal* tbmin = tile->header->bmin;
 		const dtReal* tbmax = tile->header->bmax;
-		const dtReal qfac = tile->header->bvQuantFactor;
 		
 		// Calculate quantized box
 		unsigned short bmin[3], bmax[3];
@@ -1612,12 +1615,12 @@ int dtNavMesh::queryPolygonsInTile(const dtMeshTile* tile, const dtReal* qmin, c
 		dtReal maxz = dtClamp(qmax[2], tbmin[2], tbmax[2]) - tbmin[2];
 
 		// Quantize
-		bmin[0] = (unsigned short)(qfac * minx) & 0xfffe;
-		bmin[1] = (unsigned short)(qfac * miny) & 0xfffe;
-		bmin[2] = (unsigned short)(qfac * minz) & 0xfffe;
-		bmax[0] = (unsigned short)(qfac * maxx + 1) | 1;
-		bmax[1] = (unsigned short)(qfac * maxy + 1) | 1;
-		bmax[2] = (unsigned short)(qfac * maxz + 1) | 1;
+		bmin[0] = (unsigned short)(m_bvQuantFactor * minx) & 0xfffe;
+		bmin[1] = (unsigned short)(m_bvQuantFactor * miny) & 0xfffe;
+		bmin[2] = (unsigned short)(m_bvQuantFactor * minz) & 0xfffe;
+		bmax[0] = (unsigned short)(m_bvQuantFactor * maxx + 1) | 1;
+		bmax[1] = (unsigned short)(m_bvQuantFactor * maxy + 1) | 1;
+		bmax[2] = (unsigned short)(m_bvQuantFactor * maxz + 1) | 1;
 		
 		// Traverse tree
 		dtPolyRef base = getPolyRefBase(tile);
@@ -1700,8 +1703,6 @@ dtStatus dtNavMesh::addTile(unsigned char* data, int dataSize, int flags,
 {
 	// Make sure the data is in right format.
 	dtMeshHeader* header = (dtMeshHeader*)data;
-	if (header->magic != DT_NAVMESH_MAGIC)
-		return DT_FAILURE | DT_WRONG_MAGIC;
 	if (header->version != DT_NAVMESH_VERSION)
 		return DT_FAILURE | DT_WRONG_VERSION;
 		
@@ -1841,7 +1842,7 @@ dtStatus dtNavMesh::addTile(unsigned char* data, int dataSize, int flags,
 	baseOffMeshLinks(tile);
 	
 #if WITH_NAVMESH_SEGMENT_LINKS
-	dtOffMeshSegmentData* segList = initSegmentIntersection(tile);
+	dtOffMeshSegmentData* segList = initSegmentIntersection(*this, tile);
 #endif // WITH_NAVMESH_SEGMENT_LINKS
 
 	// Create connections with neighbour tiles.
@@ -1860,7 +1861,7 @@ dtStatus dtNavMesh::addTile(unsigned char* data, int dataSize, int flags,
 			connectExtLinks(tile, neis[j], -1, bHasClusters);
 			connectExtLinks(neis[j], tile, -1, bHasClusters);
 #if WITH_NAVMESH_SEGMENT_LINKS
-			appendSegmentIntersection(segList, tile, neis[j]);
+			appendSegmentIntersection(*this, segList, tile, neis[j]);
 #endif // WITH_NAVMESH_SEGMENT_LINKS
 			connectExtOffMeshLinks(tile, neis[j], -1, bHasClusters);
 		}
@@ -1884,7 +1885,7 @@ dtStatus dtNavMesh::addTile(unsigned char* data, int dataSize, int flags,
 				connectExtLinks(neis[j], tile, dtOppositeTile(i), bHasClusters);
 			}
 #if WITH_NAVMESH_SEGMENT_LINKS
-			appendSegmentIntersection(segList, tile, neis[j]);
+			appendSegmentIntersection(*this, segList, tile, neis[j]);
 #endif // WITH_NAVMESH_SEGMENT_LINKS
 
 			connectExtOffMeshLinks(tile, neis[j], i, bHasClusters);
@@ -2335,8 +2336,6 @@ dtStatus dtNavMesh::restoreTileState(dtMeshTile* tile, const unsigned char* data
 	const dtPolyState* polyStates = (const dtPolyState*)data; data += dtAlign(sizeof(dtPolyState) * tile->header->polyCount);
 	
 	// Check that the restore is possible.
-	if (tileState->magic != DT_NAVMESH_STATE_MAGIC)
-		return DT_FAILURE | DT_WRONG_MAGIC;
 	if (tileState->version != DT_NAVMESH_STATE_VERSION)
 		return DT_FAILURE | DT_WRONG_VERSION;
 	if (tileState->ref != getTileRef(tile))
diff --git a/Engine/Source/Runtime/Navmesh/Private/Detour/DetourNavMeshBuilder.cpp b/Engine/Source/Runtime/Navmesh/Private/Detour/DetourNavMeshBuilder.cpp
--- a/Engine/Source/Runtime/Navmesh/Private/Detour/DetourNavMeshBuilder.cpp
+++ b/Engine/Source/Runtime/Navmesh/Private/Detour/DetourNavMeshBuilder.cpp
@@ -22,6 +22,7 @@
 #include "Detour/DetourNavMeshBuilder.h"
 #include "Detour/DetourNavMesh.h"
 #include "Detour/DetourCommon.h"
+#include <limits>
 
 static unsigned short MESH_NULL_IDX = 0xffff;
 
@@ -522,42 +523,68 @@ bool dtCreateNavMeshData(dtNavMeshCreateParams* params, unsigned char** outData,
 	unsigned short* polyClusters = (unsigned short*)d; d += polyClustersSize;
 #endif // WITH_NAVMESH_CLUSTER_LINKS
 	//@UE END
-	
-	// Store header
-	header->magic = DT_NAVMESH_MAGIC;
+
+	//@UE BEGIN Memory optimization
+	auto setHeaderShort = [](const int value, unsigned short& headerVal, const TCHAR* text, bool& bAllValuesSet)
+	{
+		dtAssert(value >= 0);
+
+		//check value will fit in to an unsigned short
+		if (value > std::numeric_limits<unsigned short>::max())
+		{
+			UE_LOG(LogDetour, Error, TEXT("dtCreateNavMeshData: %s is too large to fit in an unsigned short!"), *text);
+
+			bAllValuesSet = false;
+			headerVal = 0;
+		}
+		else
+		{
+			headerVal = (unsigned short)value;
+		}
+	};
+
+	bool bAllValuesSet = true;
+
 	header->version = DT_NAVMESH_VERSION;
 	header->x = params->tileX;
 	header->y = params->tileY;
-	header->layer = params->tileLayer;
-	header->userId = params->userId;
-	header->polyCount = totPolyCount;
-	header->vertCount = totVertCount;
-	header->maxLinkCount = maxLinkCount;
+	setHeaderShort(params->tileLayer, header->layer, TEXT("tile layer"), bAllValuesSet);
+	setHeaderShort(totPolyCount, header->polyCount, TEXT("poly count"), bAllValuesSet);
+	setHeaderShort(totVertCount, header->vertCount, TEXT("vert count"), bAllValuesSet);
+	setHeaderShort(maxLinkCount, header->maxLinkCount, TEXT("max link count"), bAllValuesSet);
 	dtVcopy(header->bmin, params->bmin);
 	dtVcopy(header->bmax, params->bmax);
-	header->detailMeshCount = params->polyCount;
-	header->detailVertCount = uniqueDetailVertCount;
-	header->detailTriCount = detailTriCount;
-	header->bvQuantFactor = 1.0f / params->cs;
-	header->offMeshBase = params->polyCount;
-	header->walkableHeight = params->walkableHeight;
-	header->walkableRadius = params->walkableRadius;
-	header->walkableClimb = params->walkableClimb;
-	header->offMeshConCount = storedOffMeshConCount;
-	header->bvNodeCount = params->buildBvTree ? params->polyCount*2 : 0;
+	setHeaderShort(params->polyCount, header->detailMeshCount, TEXT("detail mesh count"), bAllValuesSet);
+	setHeaderShort(uniqueDetailVertCount, header->detailVertCount, TEXT("detail vert count"), bAllValuesSet);
+	setHeaderShort(detailTriCount, header->detailTriCount, TEXT("detail tri count"), bAllValuesSet);
+	setHeaderShort(params->polyCount, header->offMeshBase, TEXT("off mesh base"), bAllValuesSet);
+	setHeaderShort(storedOffMeshConCount, header->offMeshConCount, TEXT("off mesh connection count"), bAllValuesSet);
+	setHeaderShort(params->buildBvTree ? params->polyCount * 2 : 0, header->bvNodeCount, TEXT("bv node count"), bAllValuesSet);
 
 	//@UE BEGIN
 #if WITH_NAVMESH_SEGMENT_LINKS
-	header->offMeshSegPolyBase = firstSegPoly;
-	header->offMeshSegVertBase = firstSegVert;
-	header->offMeshSegConCount = storedOffMeshSegCount;
+	setHeaderShort(firstSegPoly, header->offMeshSegPolyBase, TEXT("off mesh seg poly base"), bAllValuesSet);
+	setHeaderShort(firstSegVert, header->offMeshSegVertBase, TEXT("off mesh seg vert base"), bAllValuesSet);
+	setHeaderShort(storedOffMeshSegCount, header->offMeshSegConCount, TEXT("off mesh seg con count"), bAllValuesSet);
 #endif // WITH_NAVMESH_SEGMENT_LINKS
 
 #if WITH_NAVMESH_CLUSTER_LINKS
-	header->clusterCount = params->clusterCount; 
+	setHeaderShort(params->clusterCount, header->clusterCount, TEXT("cluster count"), bAllValuesSet);
 #endif // WITH_NAVMESH_CLUSTER_LINKS
 	//@UE END
 
+	if (!bAllValuesSet)
+	{
+		UE_LOG(LogDetour, Error, TEXT("dtCreateNavMeshData: Tile %d, %d: Layer %d: Some dtMeshHeader values are too large to fit in an unsigned short, tile not Generated! Please decrease the size of your nav mesh tiles."), params->tileX, params->tileY, params->tileLayer);
+
+		dtFree(offMeshConClass, DT_ALLOC_TEMP);
+		dtFree(data, DT_ALLOC_PERM_TILE_DATA);
+
+		return false;
+	}
+	//@UE END Memory optimization
+
+
 	const int offMeshVertsBase = params->vertCount;
 	const int offMeshPolyBase = params->polyCount;
 
@@ -826,23 +853,19 @@ bool dtNavMeshHeaderSwapEndian(unsigned char* data, const int /*dataSize*/)
 {
 	dtMeshHeader* header = (dtMeshHeader*)data;
 	
-	int swappedMagic = DT_NAVMESH_MAGIC;
 	int swappedVersion = DT_NAVMESH_VERSION;
-	dtSwapEndian(&swappedMagic);
 	dtSwapEndian(&swappedVersion);
 	
-	if ((header->magic != DT_NAVMESH_MAGIC || header->version != DT_NAVMESH_VERSION) &&
-		(header->magic != swappedMagic || header->version != swappedVersion))
+	if ((header->version != DT_NAVMESH_VERSION) &&
+		(header->version != swappedVersion))
 	{
 		return false;
 	}
 		
-	dtSwapEndian(&header->magic);
 	dtSwapEndian(&header->version);
 	dtSwapEndian(&header->x);
 	dtSwapEndian(&header->y);
 	dtSwapEndian(&header->layer);
-	dtSwapEndian(&header->userId);
 	dtSwapEndian(&header->polyCount);
 	dtSwapEndian(&header->vertCount);
 	dtSwapEndian(&header->maxLinkCount);
@@ -852,16 +875,12 @@ bool dtNavMeshHeaderSwapEndian(unsigned char* data, const int /*dataSize*/)
 	dtSwapEndian(&header->bvNodeCount);
 	dtSwapEndian(&header->offMeshConCount);
 	dtSwapEndian(&header->offMeshBase);
-	dtSwapEndian(&header->walkableHeight);
-	dtSwapEndian(&header->walkableRadius);
-	dtSwapEndian(&header->walkableClimb);
 	dtSwapEndian(&header->bmin[0]);
 	dtSwapEndian(&header->bmin[1]);
 	dtSwapEndian(&header->bmin[2]);
 	dtSwapEndian(&header->bmax[0]);
 	dtSwapEndian(&header->bmax[1]);
 	dtSwapEndian(&header->bmax[2]);
-	dtSwapEndian(&header->bvQuantFactor);
 
 	//@UE BEGIN
 #if WITH_NAVMESH_SEGMENT_LINKS
@@ -890,8 +909,6 @@ bool dtNavMeshDataSwapEndian(unsigned char* data, const int /*dataSize*/)
 {
 	// Make sure the data is in right format.
 	dtMeshHeader* header = (dtMeshHeader*)data;
-	if (header->magic != DT_NAVMESH_MAGIC)
-		return false;
 	if (header->version != DT_NAVMESH_VERSION)
 		return false;
 	
@@ -1028,12 +1045,10 @@ bool dtNavMeshDataSwapEndian(unsigned char* data, const int /*dataSize*/)
 
 // @UE BEGIN
 // Experimental tile transform
-bool dtTransformTileData(unsigned char* data, const int dataSize, const int offsetX, const int offsetY, const dtReal tileWidth, const dtReal tileHeight, const dtReal rotationDeg)
+bool dtTransformTileData(unsigned char* data, const int dataSize, const int offsetX, const int offsetY, const dtReal tileWidth, const dtReal tileHeight, const dtReal rotationDeg, const dtReal bvQuantFactor)
 {
 	// Make sure the data is in right format.
 	dtMeshHeader* header = (dtMeshHeader*)data;
-	if (header->magic != DT_NAVMESH_MAGIC)
-		return false;
 	if (header->version != DT_NAVMESH_VERSION)
 		return false;
 
@@ -1095,15 +1110,14 @@ bool dtTransformTileData(unsigned char* data, const int dataSize, const int offs
 	offset[2] = tileHeight * offsetY;
 
 	// Compute center for bvtree
-	const dtReal qfac = tile.header->bvQuantFactor;
 	unsigned short qMin[3];
 	unsigned short qMax[3];
-	qMin[0] = (unsigned short)(qfac * tile.header->bmin[0]);
-	qMin[1] = (unsigned short)(qfac * tile.header->bmin[1]);
-	qMin[2] = (unsigned short)(qfac * tile.header->bmin[2]);
-	qMax[0] = (unsigned short)(qfac * tile.header->bmax[0]);
-	qMax[1] = (unsigned short)(qfac * tile.header->bmax[1]);
-	qMax[2] = (unsigned short)(qfac * tile.header->bmax[2]);
+	qMin[0] = (unsigned short)(bvQuantFactor * tile.header->bmin[0]);
+	qMin[1] = (unsigned short)(bvQuantFactor * tile.header->bmin[1]);
+	qMin[2] = (unsigned short)(bvQuantFactor * tile.header->bmin[2]);
+	qMax[0] = (unsigned short)(bvQuantFactor * tile.header->bmax[0]);
+	qMax[1] = (unsigned short)(bvQuantFactor * tile.header->bmax[1]);
+	qMax[2] = (unsigned short)(bvQuantFactor * tile.header->bmax[2]);
 	unsigned short qLocalCenter[3];
 	qLocalCenter[0] = (qMax[0] - qMin[0])/2;
 	qLocalCenter[1] = (qMax[1] - qMin[1])/2;
diff --git a/Engine/Source/Runtime/Navmesh/Private/Detour/DetourNavMeshQuery.cpp b/Engine/Source/Runtime/Navmesh/Private/Detour/DetourNavMeshQuery.cpp
--- a/Engine/Source/Runtime/Navmesh/Private/Detour/DetourNavMeshQuery.cpp
+++ b/Engine/Source/Runtime/Navmesh/Private/Detour/DetourNavMeshQuery.cpp
@@ -1297,8 +1297,8 @@ int dtNavMeshQuery::queryPolygonsInTile(const dtMeshTile* tile, const dtReal* qm
 		const dtBVNode* end = &tile->bvTree[tile->header->bvNodeCount];
 		const dtReal* tbmin = tile->header->bmin;
 		const dtReal* tbmax = tile->header->bmax;
-		const dtReal qfac = tile->header->bvQuantFactor;
-		
+		const dtReal qfac = m_nav->getBVQuantFactor();
+
 		// Calculate quantized box
 		unsigned short bmin[3], bmax[3];
 		// dtClamp query box to world box.
diff --git a/Engine/Source/Runtime/Navmesh/Private/DetourCrowd/DetourObstacleAvoidance.cpp b/Engine/Source/Runtime/Navmesh/Private/DetourCrowd/DetourObstacleAvoidance.cpp
--- a/Engine/Source/Runtime/Navmesh/Private/DetourCrowd/DetourObstacleAvoidance.cpp
+++ b/Engine/Source/Runtime/Navmesh/Private/DetourCrowd/DetourObstacleAvoidance.cpp
@@ -205,13 +205,13 @@ void dtFreeObstacleAvoidanceQuery(dtObstacleAvoidanceQuery* ptr)
 
 
 dtObstacleAvoidanceQuery::dtObstacleAvoidanceQuery() :
-	m_maxPatterns(0),
 	m_customPatterns(0),
-	m_maxCircles(0),
 	m_circles(0),
+	m_segments(0),
+	m_maxPatterns(0),
+	m_maxCircles(0),
 	m_ncircles(0),
 	m_maxSegments(0),
-	m_segments(0),
 	m_nsegments(0)
 {
 }
diff --git a/Engine/Source/Runtime/Navmesh/Private/DetourCrowd/DetourPathQueue.cpp b/Engine/Source/Runtime/Navmesh/Private/DetourCrowd/DetourPathQueue.cpp
--- a/Engine/Source/Runtime/Navmesh/Private/DetourCrowd/DetourPathQueue.cpp
+++ b/Engine/Source/Runtime/Navmesh/Private/DetourCrowd/DetourPathQueue.cpp
@@ -23,10 +23,10 @@
 
 
 dtPathQueue::dtPathQueue() :
+	m_navquery(0),
 	m_nextHandle(1),
 	m_maxPathSize(0),
-	m_queueHead(0),
-	m_navquery(0)
+	m_queueHead(0)
 {
 	for (int i = 0; i < MAX_QUEUE; ++i)
 		m_queue[i].path = 0;
diff --git a/Engine/Source/Runtime/Navmesh/Private/DetourTileCache/DetourTileCache.cpp b/Engine/Source/Runtime/Navmesh/Private/DetourTileCache/DetourTileCache.cpp
--- a/Engine/Source/Runtime/Navmesh/Private/DetourTileCache/DetourTileCache.cpp
+++ b/Engine/Source/Runtime/Navmesh/Private/DetourTileCache/DetourTileCache.cpp
@@ -295,8 +295,6 @@ dtStatus dtTileCache::addTile(unsigned char* data, const int dataSize, unsigned
 {
 	// Make sure the data is in right format.
 	dtTileCacheLayerHeader* header = (dtTileCacheLayerHeader*)data;
-	if (header->magic != DT_TILECACHE_MAGIC)
-		return DT_FAILURE | DT_WRONG_MAGIC;
 	if (header->version != DT_TILECACHE_VERSION)
 		return DT_FAILURE | DT_WRONG_VERSION;
 	
diff --git a/Engine/Source/Runtime/Navmesh/Private/DetourTileCache/DetourTileCacheBuilder.cpp b/Engine/Source/Runtime/Navmesh/Private/DetourTileCache/DetourTileCacheBuilder.cpp
--- a/Engine/Source/Runtime/Navmesh/Private/DetourTileCache/DetourTileCacheBuilder.cpp
+++ b/Engine/Source/Runtime/Navmesh/Private/DetourTileCache/DetourTileCacheBuilder.cpp
@@ -2678,8 +2678,6 @@ dtStatus dtDecompressTileCacheLayer(dtTileCacheAlloc* alloc, dtTileCacheCompress
 	*layerOut = 0;
 
 	dtTileCacheLayerHeader* compressedHeader = (dtTileCacheLayerHeader*)compressed;
-	if (compressedHeader->magic != DT_TILECACHE_MAGIC)
-		return DT_FAILURE | DT_WRONG_MAGIC;
 	if (compressedHeader->version != DT_TILECACHE_VERSION)
 		return DT_FAILURE | DT_WRONG_VERSION;
 	
@@ -2732,13 +2730,12 @@ bool dtTileCacheHeaderSwapEndian(unsigned char* data, const int dataSize)
 	dtSwapEndian(&swappedMagic);
 	dtSwapEndian(&swappedVersion);
 	
-	if ((header->magic != DT_TILECACHE_MAGIC || header->version != DT_TILECACHE_VERSION) &&
-		(header->magic != swappedMagic || header->version != swappedVersion))
+	if ((header->version != DT_TILECACHE_VERSION) &&
+		(header->version != swappedVersion))
 	{
 		return false;
 	}
 	
-	dtSwapEndian(&header->magic);
 	dtSwapEndian(&header->version);
 	dtSwapEndian(&header->tx);
 	dtSwapEndian(&header->ty);
diff --git a/Engine/Source/Runtime/Navmesh/Public/Detour/DetourNavMesh.h b/Engine/Source/Runtime/Navmesh/Public/Detour/DetourNavMesh.h
--- a/Engine/Source/Runtime/Navmesh/Public/Detour/DetourNavMesh.h
+++ b/Engine/Source/Runtime/Navmesh/Public/Detour/DetourNavMesh.h
@@ -27,6 +27,7 @@
 #include "Detour/DetourStatus.h"
 #include "DetourLargeWorldCoordinates.h"
 
+NAVMESH_API DECLARE_LOG_CATEGORY_EXTERN(LogDetour, Log, All);
 
 // Note: If you want to use 64-bit refs, change the types of both dtPolyRef & dtTileRef.
 // It is also recommended that you change dtHashRef() to a proper 64-bit hash.
@@ -82,7 +83,8 @@ static const int DT_VERTS_PER_POLYGON = 6;
 ///
 
 /// A magic number used to detect compatibility of navigation tile data.
-static const int DT_NAVMESH_MAGIC = 'D'<<24 | 'N'<<16 | 'A'<<8 | 'V';
+/// UE magic removed to save memory
+//static const int DT_NAVMESH_MAGIC = 'D'<<24 | 'N'<<16 | 'A'<<8 | 'V';
 
 /// A version number used to detect compatibility of navigation tile data.
 static const int DT_NAVMESH_VERSION = 7;
@@ -233,8 +235,8 @@ struct dtPoly
 /// Defines the location of detail sub-mesh data within a dtMeshTile.
 struct dtPolyDetail
 {
-	unsigned int vertBase;			///< The offset of the vertices in the dtMeshTile::detailVerts array.
-	unsigned int triBase;			///< The offset of the triangles in the dtMeshTile::detailTris array.
+	unsigned short vertBase;		///< The offset of the vertices in the dtMeshTile::detailVerts array.
+	unsigned short triBase;			///< The offset of the triangles in the dtMeshTile::detailTris array.
 	unsigned char vertCount;		///< The number of vertices in the sub-mesh.
 	unsigned char triCount;			///< The number of triangles in the sub-mesh.
 };
@@ -368,52 +370,48 @@ struct dtClusterLink
 /// @ingroup detour
 struct dtMeshHeader
 {
-	int magic;				///< Tile magic number. (Used to identify the data format.)
-	int version;			///< Tile data format version number.
+	unsigned short version;	///< Tile data format version number.
+	unsigned short layer;			///< The layer of the tile within the dtNavMesh tile grid. (x, y, layer)
+
+	unsigned short polyCount;		///< The number of polygons in the tile.
+	unsigned short vertCount;		///< The number of vertices in the tile.
+
 	int x;					///< The x-position of the tile within the dtNavMesh tile grid. (x, y, layer)
 	int y;					///< The y-position of the tile within the dtNavMesh tile grid. (x, y, layer)
-	int layer;				///< The layer of the tile within the dtNavMesh tile grid. (x, y, layer)
-	unsigned int userId;	///< The user defined id of the tile.
-	int polyCount;			///< The number of polygons in the tile.
-	int vertCount;			///< The number of vertices in the tile.
-	int maxLinkCount;		///< The number of allocated links.
-	int detailMeshCount;	///< The number of sub-meshes in the detail mesh.
+			
+	unsigned short maxLinkCount;		///< The number of allocated links.
+	unsigned short detailMeshCount;		///< The number of sub-meshes in the detail mesh.
 	
 	/// The number of unique vertices in the detail mesh. (In addition to the polygon vertices.)
-	int detailVertCount;
+	unsigned short detailVertCount;
 	
-	int detailTriCount;			///< The number of triangles in the detail mesh.
-	int bvNodeCount;			///< The number of bounding volume nodes. (Zero if bounding volumes are disabled.)
-	int offMeshConCount;		///< The number of point type off-mesh connections.
-	int offMeshBase;			///< The index of the first polygon which is an point type off-mesh connection.
+	unsigned short detailTriCount;		///< The number of triangles in the detail mesh.
+	unsigned short bvNodeCount;			///< The number of bounding volume nodes. (Zero if bounding volumes are disabled.)
+	unsigned short offMeshConCount;		///< The number of point type off-mesh connections.
+	unsigned short offMeshBase;			///< The index of the first polygon which is an point type off-mesh connection.
 
 	//@UE BEGIN
 #if WITH_NAVMESH_SEGMENT_LINKS
-	int offMeshSegConCount;		///< The number of segment type off-mesh connections.
-	int offMeshSegPolyBase;		///< The index of the first polygon which is an segment type off-mesh connection
-	int offMeshSegVertBase;		///< The index of the first vertex used by segment type off-mesh connection
+	unsigned short offMeshSegConCount;	///< The number of segment type off-mesh connections.
+	unsigned short offMeshSegPolyBase;	///< The index of the first polygon which is an segment type off-mesh connection
+	unsigned short offMeshSegVertBase;	///< The index of the first vertex used by segment type off-mesh connection
 #endif // WITH_NAVMESH_SEGMENT_LINKS
 
 #if WITH_NAVMESH_CLUSTER_LINKS
-	int clusterCount;			///< Number of clusters
+	unsigned short clusterCount;			///< Number of clusters
 #endif // WITH_NAVMESH_CLUSTER_LINKS
 	//@UE END
 
-	dtReal walkableHeight;		///< The height of the agents using the tile.
-	dtReal walkableRadius;		///< The radius of the agents using the tile.
-	dtReal walkableClimb;		///< The maximum climb height of the agents using the tile.
+	// These should be at the bottom, as they are less often used than the rest of the data. The rest will fit in one cache line.
 	dtReal bmin[3];				///< The minimum bounds of the tile's AABB. [(x, y, z)]
 	dtReal bmax[3];				///< The maximum bounds of the tile's AABB. [(x, y, z)]
-	
-	/// The bounding volume quantization factor. 
-	dtReal bvQuantFactor;
 };
 
 /// Defines a navigation mesh tile.
 /// @ingroup detour
 struct dtMeshTile
 {
-	unsigned int salt;					///< Counter describing modifications to the tile.
+	unsigned int salt;						///< Counter describing modifications to the tile.
 
 	unsigned int linksFreeList;			///< Index to the next free link.
 	dtMeshHeader* header;				///< The tile header.
@@ -449,11 +447,15 @@ struct dtMeshTile
 	unsigned short* polyClusters;				///< Cluster Id for each ground type polygon [Size: dtMeshHeader::polyCount]
 
 	dtChunkArray<dtClusterLink, DT_ALLOC_PERM_TILE_DYNLINK_CLUSTER> dynamicLinksC; ///< Dynamic links array (indices starting from DT_CLINK_FIRST)
-	unsigned int dynamicFreeListC;				///< Index of the next free dynamic link
 #endif // WITH_NAVMESH_CLUSTER_LINKS
 
 	dtChunkArray<dtLink, DT_ALLOC_PERM_TILE_DYNLINK_OFFMESH> dynamicLinksO; ///< Dynamic links array (indices starting from dtMeshHeader::maxLinkCount)
 	unsigned int dynamicFreeListO;				///< Index of the next free dynamic link
+
+#if WITH_NAVMESH_CLUSTER_LINKS
+	unsigned int dynamicFreeListC;				///< Index of the next free dynamic link
+#endif // WITH_NAVMESH_CLUSTER_LINKS
+
 	//@UE END
 };
 
@@ -839,6 +841,18 @@ public:
 	//@UE END
 	/// @}
 
+	//@UE Begin LWCoords
+	dtReal getWalkableHeight() const { return m_walkableHeight; }
+	dtReal getWalkableRadius() const { return m_walkableRadius; }
+	dtReal getWalkableClimb() const { return m_walkableClimb; }
+	dtReal getBVQuantFactor() const { return m_bvQuantFactor; }
+
+	void setWalkableHeight(dtReal height) { m_walkableHeight = height; }
+	void setWalkableRadius(dtReal radius) { m_walkableRadius = radius; }
+	void setWalkableClimb(dtReal climb) { m_walkableClimb = climb; }
+	void setBVQuantFactor(dtReal factor) { m_bvQuantFactor = factor; }
+	//@UE END LWCoords
+
 private:
 
 	// [UE] result struct for findConnectingPolys
@@ -908,6 +922,14 @@ private:
 	dtNavMeshParams m_params;			///< Current initialization params. TODO: do not store this info twice.
 	dtReal m_orig[3];					///< Origin of the tile (0,0)
 	dtReal m_tileWidth, m_tileHeight;	///< Dimensions of each tile.
+
+	//@UE BEGIN Memory optimization
+	dtReal m_walkableHeight;				///< The height of the agents using the tile.
+	dtReal m_walkableRadius;				///< The radius of the agents using the tile.
+	dtReal m_walkableClimb;				///< The maximum climb height of the agents using the tile.
+	dtReal m_bvQuantFactor;				///<  The bounding volume quantization factor. 
+	//@UE END Memory optimization
+
 	int m_maxTiles;						///< Max number of tiles.
 	int m_tileLutSize;					///< Tile hash lookup size (must be pot).
 	int m_tileLutMask;					///< Tile hash lookup mask.
@@ -1023,13 +1045,13 @@ if (neis[n] & #DT_EX_LINK)
 }
 @endcode
 
-@var dtReal dtMeshHeader::bvQuantFactor
+@var dtReal dtNavMesh::bvQuantFactor
 @par
 
 This value is used for converting between world and bounding volume coordinates.
 For example:
 @code
-const dtReal cs = 1.0f / tile->header->bvQuantFactor;
+const dtReal cs = 1.0f / bvQuantFactor;
 const dtBVNode* n = &tile->bvTree[i];
 if (n->i >= 0)
 {
diff --git a/Engine/Source/Runtime/Navmesh/Public/Detour/DetourNavMeshBuilder.h b/Engine/Source/Runtime/Navmesh/Public/Detour/DetourNavMeshBuilder.h
--- a/Engine/Source/Runtime/Navmesh/Public/Detour/DetourNavMeshBuilder.h
+++ b/Engine/Source/Runtime/Navmesh/Public/Detour/DetourNavMeshBuilder.h
@@ -177,7 +177,7 @@ NAVMESH_API bool dtNavMeshDataSwapEndian(unsigned char* data, const int dataSize
 ///  @param[in]		tileWidth		Tile width.
 ///  @param[in]		tileHeight		Tile height.
 ///  @param[in]		rotationDeg		Rotation in degrees.
-NAVMESH_API bool dtTransformTileData(unsigned char* data, const int dataSize, const int offsetX, const int offsetY, const dtReal tileWidth, const dtReal tileHeight, const dtReal rotationDeg);
+NAVMESH_API bool dtTransformTileData(unsigned char* data, const int dataSize, const int offsetX, const int offsetY, const dtReal tileWidth, const dtReal tileHeight, const dtReal rotationDeg, const dtReal bvQuantFactor);
 
 /// Compute XY offset caused by the given rotation
 ///  @param[in]		position		Position to rotate. [(x, y, z)]
diff --git a/Engine/Source/Runtime/Navmesh/Public/Detour/DetourStatus.h b/Engine/Source/Runtime/Navmesh/Public/Detour/DetourStatus.h
--- a/Engine/Source/Runtime/Navmesh/Public/Detour/DetourStatus.h
+++ b/Engine/Source/Runtime/Navmesh/Public/Detour/DetourStatus.h
@@ -33,7 +33,7 @@ static const unsigned int DT_IN_PROGRESS = 1u << 29;		// Operation still in prog
 
 // Detail information for status.
 static const unsigned int DT_STATUS_DETAIL_MASK = 0x0ffffff;
-static const unsigned int DT_WRONG_MAGIC = 1 << 0;		// Input data is not recognized.
+//static const unsigned int DT_WRONG_MAGIC = 1 << 0;		// Input data is not recognized. NOTE: UE magic removed to save memory
 static const unsigned int DT_WRONG_VERSION = 1 << 1;	// Input data is in wrong version.
 static const unsigned int DT_OUT_OF_MEMORY = 1 << 2;	// Operation ran out of memory.
 static const unsigned int DT_INVALID_PARAM = 1 << 3;	// An input parameter was invalid.
diff --git a/Engine/Source/Runtime/Navmesh/Public/DetourCrowd/DetourCrowd.h b/Engine/Source/Runtime/Navmesh/Public/DetourCrowd/DetourCrowd.h
--- a/Engine/Source/Runtime/Navmesh/Public/DetourCrowd/DetourCrowd.h
+++ b/Engine/Source/Runtime/Navmesh/Public/DetourCrowd/DetourCrowd.h
@@ -146,9 +146,6 @@ struct NAVMESH_API dtCrowdAgent
 	/// The known neighbors of the agent.
 	dtCrowdNeighbour neis[DT_CROWDAGENT_MAX_NEIGHBOURS];
 
-	/// The number of neighbors.
-	int nneis;
-	
 	/// The desired speed.
 	dtReal desiredSpeed;
 
@@ -167,16 +164,20 @@ struct NAVMESH_API dtCrowdAgent
 	/// The reference id of the polygon being entered at the corner. [(polyRef) * #ncorners]
 	dtPolyRef cornerPolys[DT_CROWDAGENT_MAX_CORNERS];
 
+	dtReal targetReplanTime;				/// <Time since the agent's target was replanned.
+	dtReal targetPos[3];					///< Target position of the movement request (or velocity in case of DT_CROWDAGENT_TARGET_VELOCITY).
+	dtPolyRef targetRef;				///< Target polyref of the movement request.
+	dtPathQueueRef targetPathqRef;		///< Path finder ref.
+
+	/// The number of neighbors.
+	int nneis;
+
 	/// The local path corridor corner flags. (See: #dtStraightPathFlags) [(flags) * #ncorners]
 	unsigned char cornerFlags[DT_CROWDAGENT_MAX_CORNERS];
 
 	/// The number of corners.
 	int ncorners;
 	
-	dtReal targetReplanTime;				/// <Time since the agent's target was replanned.
-	dtReal targetPos[3];					///< Target position of the movement request (or velocity in case of DT_CROWDAGENT_TARGET_VELOCITY).
-	dtPolyRef targetRef;				///< Target polyref of the movement request.
-	dtPathQueueRef targetPathqRef;		///< Path finder ref.
 	unsigned char targetReplan;			///< Flag indicating that the current path is being replanned.
 	unsigned char targetState;			///< State of the movement request.
 
@@ -243,7 +244,6 @@ class NAVMESH_API dtCrowd
 	dtProximityGrid* m_grid;
 	
 	dtPolyRef* m_pathResult;
-	int m_maxPathResult;
 	
 	dtReal m_ext[3];
 	dtQueryFilter m_filters[DT_CROWD_MAX_FILTERS];
@@ -260,6 +260,8 @@ class NAVMESH_API dtCrowd
 	// [UE] separation filter
 	dtReal m_separationDirFilter;
 
+	int m_maxPathResult;
+
 	int m_velocitySampleCount;
 
 	dtNavMeshQuery* m_navquery;
diff --git a/Engine/Source/Runtime/Navmesh/Public/DetourCrowd/DetourObstacleAvoidance.h b/Engine/Source/Runtime/Navmesh/Public/DetourCrowd/DetourObstacleAvoidance.h
--- a/Engine/Source/Runtime/Navmesh/Public/DetourCrowd/DetourObstacleAvoidance.h
+++ b/Engine/Source/Runtime/Navmesh/Public/DetourCrowd/DetourObstacleAvoidance.h
@@ -177,15 +177,16 @@ private:
 	dtReal m_vmax;
 	dtReal m_invVmax;
 
-	int m_maxPatterns;
 	dtObstacleAvoidancePattern* m_customPatterns;
+	dtObstacleCircle* m_circles;
+	dtObstacleSegment* m_segments;
+
+	int m_maxPatterns;
 
 	int m_maxCircles;
-	dtObstacleCircle* m_circles;
 	int m_ncircles;
 
 	int m_maxSegments;
-	dtObstacleSegment* m_segments;
 	int m_nsegments;
 };
 
diff --git a/Engine/Source/Runtime/Navmesh/Public/DetourTileCache/DetourTileCacheBuilder.h b/Engine/Source/Runtime/Navmesh/Public/DetourTileCache/DetourTileCacheBuilder.h
--- a/Engine/Source/Runtime/Navmesh/Public/DetourTileCache/DetourTileCacheBuilder.h
+++ b/Engine/Source/Runtime/Navmesh/Public/DetourTileCache/DetourTileCacheBuilder.h
@@ -36,13 +36,13 @@ static const unsigned short DT_TILECACHE_NULL_IDX = 0xffff;
 
 struct dtTileCacheLayerHeader
 {
-	int magic;								///< Data magic
-	int version;							///< Data version
-	int tx,ty,tlayer;
-	dtReal bmin[3], bmax[3];
+	short version;							///< Data version
+	int tx, ty, tlayer;
 	unsigned short hmin, hmax;				///< Height min/max range
 	unsigned short width, height;			///< Dimension of the layer.
 	unsigned short minx, maxx, miny, maxy;	///< Usable sub-region.
+	// These should be at the bottom, as they are less often used than the rest of the data.
+	dtReal bmin[3], bmax[3];
 };
 
 struct dtTileCacheLayer
