[Navmesh] Moving the initialization of dtNavMesh parameters into the init() and serializing the parameters in ARecastNavMesh.
(for safer initialization flow and fix potential mismatch between navmesh actor and serialized navmesh data)

diff --git a/Engine/Source/Runtime/NavigationSystem/Private/NavMesh/PImplRecastNavMesh.cpp b/Engine/Source/Runtime/NavigationSystem/Private/NavMesh/PImplRecastNavMesh.cpp
--- a/Engine/Source/Runtime/NavigationSystem/Private/NavMesh/PImplRecastNavMesh.cpp
+++ b/Engine/Source/Runtime/NavigationSystem/Private/NavMesh/PImplRecastNavMesh.cpp
@@ -605,13 +605,6 @@ void FPImplRecastNavMesh::Serialize( FArchive& ArWrapped, int32 NavMeshVersion )
 		{
 			UE_VLOG(NavMeshOwner, LogNavigation, Error, TEXT("Failed to allocate Recast navmesh"));
 		}
-		else if (NavMeshOwner)
-		{
-			DetourNavMesh->setWalkableClimb(NavMeshOwner->AgentMaxStepHeight);
-			DetourNavMesh->setWalkableHeight(NavMeshOwner->AgentHeight);
-			DetourNavMesh->setWalkableRadius(NavMeshOwner->AgentRadius);
-			DetourNavMesh->setBVQuantFactor(1.f / NavMeshOwner->CellSize);
-		}
 	}
 
 	int32 NumTiles = 0;
@@ -668,6 +661,21 @@ void FPImplRecastNavMesh::Serialize( FArchive& ArWrapped, int32 NavMeshVersion )
 	Ar << Params.maxTiles;				///< The maximum number of tiles the navigation mesh can contain.
 	Ar << Params.maxPolys;
 
+	if(NavMeshOwner->NavMeshVersion >= NAVMESHVER_OPTIM_FIX_SERIALIZE_PARAMS)
+	{
+		Ar << Params.walkableHeight;
+		Ar << Params.walkableRadius;
+		Ar << Params.walkableClimb;
+		Ar << Params.bvQuantFactor;
+	}
+	else
+	{
+		Params.walkableHeight = NavMeshOwner->AgentHeight;
+		Params.walkableRadius = NavMeshOwner->AgentRadius;
+		Params.walkableClimb = NavMeshOwner->AgentMaxStepHeight;
+		Params.bvQuantFactor = 1.f / NavMeshOwner->CellSize;
+	}
+
 	if (Ar.IsLoading())
 	{
 		// at this point we can tell whether navmesh being loaded is in line
diff --git a/Engine/Source/Runtime/NavigationSystem/Private/NavMesh/RecastNavMesh.cpp b/Engine/Source/Runtime/NavigationSystem/Private/NavMesh/RecastNavMesh.cpp
--- a/Engine/Source/Runtime/NavigationSystem/Private/NavMesh/RecastNavMesh.cpp
+++ b/Engine/Source/Runtime/NavigationSystem/Private/NavMesh/RecastNavMesh.cpp
@@ -764,14 +764,7 @@ void ARecastNavMesh::OnRegistered()
 	check(RecastNavMeshImpl);
 
 	// This check can fail when the NavMeshVersion indicates the map needs the nav mesh rebuilt
-	if (RecastNavMeshImpl->GetRecastMesh())
-	{
-		// Set walkable climb / radius / height in recast remembering to convert to recast cell units
-		RecastNavMeshImpl->GetRecastMesh()->setWalkableClimb(AgentMaxStepHeight);
-		RecastNavMeshImpl->GetRecastMesh()->setWalkableHeight(AgentHeight);
-		RecastNavMeshImpl->GetRecastMesh()->setWalkableRadius(AgentRadius);
-		RecastNavMeshImpl->GetRecastMesh()->setBVQuantFactor(1.f / CellSize);
-	}
+	ensure(RecastNavMeshImpl->GetRecastMesh() == nullptr || RecastNavMeshImpl->GetRecastMesh()->getBVQuantFactor() != 0);
 }
 
 void ARecastNavMesh::SerializeRecastNavMesh(FArchive& Ar, FPImplRecastNavMesh*& NavMesh, int32 InNavMeshVersion)
diff --git a/Engine/Source/Runtime/NavigationSystem/Private/NavMesh/RecastNavMeshGenerator.cpp b/Engine/Source/Runtime/NavigationSystem/Private/NavMesh/RecastNavMeshGenerator.cpp
--- a/Engine/Source/Runtime/NavigationSystem/Private/NavMesh/RecastNavMeshGenerator.cpp
+++ b/Engine/Source/Runtime/NavigationSystem/Private/NavMesh/RecastNavMeshGenerator.cpp
@@ -4732,14 +4732,7 @@ void FRecastNavMeshGenerator::Init()
 		}
 	}
 
-	// This check can fail when the NavMeshVersion indicates the map needs the nav mesh rebuilt
-	if (DestNavMesh->GetRecastMesh())
-	{
-		DestNavMesh->GetRecastMesh()->setWalkableClimb(Config.AgentMaxClimb);
-		DestNavMesh->GetRecastMesh()->setWalkableHeight(Config.AgentHeight);
-		DestNavMesh->GetRecastMesh()->setWalkableRadius(Config.AgentRadius);
-		DestNavMesh->GetRecastMesh()->setBVQuantFactor(1.f / Config.cs);
-	}
+	ensure(DestNavMesh->GetRecastMesh() == nullptr || DestNavMesh->GetRecastMesh()->getBVQuantFactor() != 0);
 
 	UpdateNavigationBounds();
 
@@ -4891,10 +4884,10 @@ bool FRecastNavMeshGenerator::ConstructTiledNavMesh()
 		CalcNavMeshProperties(TiledMeshParameters.maxTiles, TiledMeshParameters.maxPolys);
 		Config.MaxPolysPerTile = TiledMeshParameters.maxPolys;
 
-		DetourMesh->setWalkableClimb(Config.AgentMaxClimb);
-		DetourMesh->setWalkableHeight(Config.AgentHeight);
-		DetourMesh->setWalkableRadius(Config.AgentRadius);
-		DetourMesh->setBVQuantFactor(1.f / Config.cs);
+		TiledMeshParameters.walkableClimb = Config.AgentMaxClimb;
+		TiledMeshParameters.walkableHeight = Config.AgentHeight;
+		TiledMeshParameters.walkableRadius = Config.AgentRadius;
+		TiledMeshParameters.bvQuantFactor = 1.f / Config.cs;
 
 		if (TiledMeshParameters.maxTiles == 0)
 		{
diff --git a/Engine/Source/Runtime/NavigationSystem/Public/NavMesh/RecastVersion.h b/Engine/Source/Runtime/NavigationSystem/Public/NavMesh/RecastVersion.h
--- a/Engine/Source/Runtime/NavigationSystem/Public/NavMesh/RecastVersion.h
+++ b/Engine/Source/Runtime/NavigationSystem/Public/NavMesh/RecastVersion.h
@@ -3,22 +3,23 @@
 #pragma once
 
 /** Versioning */
-#define NAVMESHVER_INITIAL					1
-#define NAVMESHVER_TILED_GENERATION			2
-#define NAVMESHVER_SEAMLESS_REBUILDING_1	3
-#define NAVMESHVER_AREA_CLASSES				4
-#define NAVMESHVER_CLUSTER_PATH				5
-#define NAVMESHVER_SEGMENT_LINKS			6
-#define NAVMESHVER_DYNAMIC_LINKS			7
-#define NAVMESHVER_64BIT					9
-#define NAVMESHVER_CLUSTER_SIMPLIFIED		10
-#define NAVMESHVER_OFFMESH_HEIGHT_BUG		11
-#define NAVMESHVER_LANDSCAPE_HEIGHT			13
-#define NAVMESHVER_LWCOORDS					14
-#define NAVMESHVER_OODLE_COMPRESSION		15
-#define NAVMESHVER_LWCOORDS_SEREALIZATION 	17 // Allows for nav meshes to be serialized agnostic of LWCoords being float or double.
-#define NAVMESHVER_MAXTILES_COUNT_CHANGE 	19
-#define NAVMESHVER_LWCOORDS_OPTIMIZATION	20
+#define NAVMESHVER_INITIAL						1
+#define NAVMESHVER_TILED_GENERATION				2
+#define NAVMESHVER_SEAMLESS_REBUILDING_1		3
+#define NAVMESHVER_AREA_CLASSES					4
+#define NAVMESHVER_CLUSTER_PATH					5
+#define NAVMESHVER_SEGMENT_LINKS				6
+#define NAVMESHVER_DYNAMIC_LINKS				7
+#define NAVMESHVER_64BIT						9
+#define NAVMESHVER_CLUSTER_SIMPLIFIED			10
+#define NAVMESHVER_OFFMESH_HEIGHT_BUG			11
+#define NAVMESHVER_LANDSCAPE_HEIGHT				13
+#define NAVMESHVER_LWCOORDS						14
+#define NAVMESHVER_OODLE_COMPRESSION			15
+#define NAVMESHVER_LWCOORDS_SEREALIZATION 		17 // Allows for nav meshes to be serialized agnostic of LWCoords being float or double.
+#define NAVMESHVER_MAXTILES_COUNT_CHANGE 		19
+#define NAVMESHVER_LWCOORDS_OPTIMIZATION		20
+#define NAVMESHVER_OPTIM_FIX_SERIALIZE_PARAMS	21 // Fix, serialize params that used to be in the tile and are now in the navmesh.
 
-#define NAVMESHVER_LATEST				NAVMESHVER_LWCOORDS_OPTIMIZATION
+#define NAVMESHVER_LATEST				NAVMESHVER_OPTIM_FIX_SERIALIZE_PARAMS
 #define NAVMESHVER_MIN_COMPATIBLE		NAVMESHVER_LWCOORDS_OPTIMIZATION
diff --git a/Engine/Source/Runtime/Navmesh/Private/Detour/DetourNavMesh.cpp b/Engine/Source/Runtime/Navmesh/Private/Detour/DetourNavMesh.cpp
--- a/Engine/Source/Runtime/Navmesh/Private/Detour/DetourNavMesh.cpp
+++ b/Engine/Source/Runtime/Navmesh/Private/Detour/DetourNavMesh.cpp
@@ -825,6 +825,12 @@ dtNavMesh::~dtNavMesh()
 		
 dtStatus dtNavMesh::init(const dtNavMeshParams* params)
 {
+	m_walkableHeight = params->walkableHeight;
+	m_walkableRadius = params->walkableRadius; 
+	m_walkableClimb = params->walkableClimb;	
+	m_bvQuantFactor = params->bvQuantFactor;
+	check(m_bvQuantFactor != 0);
+	
 	memcpy(&m_params, params, sizeof(dtNavMeshParams));
 	dtVcopy(m_orig, params->orig);
 	m_tileWidth = params->tileWidth;
diff --git a/Engine/Source/Runtime/Navmesh/Public/Detour/DetourNavMesh.h b/Engine/Source/Runtime/Navmesh/Public/Detour/DetourNavMesh.h
--- a/Engine/Source/Runtime/Navmesh/Public/Detour/DetourNavMesh.h
+++ b/Engine/Source/Runtime/Navmesh/Public/Detour/DetourNavMesh.h
@@ -465,6 +465,13 @@ struct dtMeshTile
 /// @ingroup detour
 struct dtNavMeshParams
 {
+	//@UE BEGIN Memory optimization
+	dtReal walkableHeight;			///< The height of the agents using the tile.
+	dtReal walkableRadius;			///< The radius of the agents using the tile.
+	dtReal walkableClimb;			///< The maximum climb height of the agents using the tile.	
+	dtReal bvQuantFactor;			///< The bounding volume quantization factor.
+	//@UE END Memory optimization 	
+	
 	dtReal orig[3];					///< The world space origin of the navigation mesh's tile space. [(x, y, z)]
 	dtReal tileWidth;				///< The width of each tile. (Along the x-axis.)
 	dtReal tileHeight;				///< The height of each tile. (Along the z-axis.)
@@ -846,11 +853,6 @@ public:
 	dtReal getWalkableRadius() const { return m_walkableRadius; }
 	dtReal getWalkableClimb() const { return m_walkableClimb; }
 	dtReal getBVQuantFactor() const { return m_bvQuantFactor; }
-
-	void setWalkableHeight(dtReal height) { m_walkableHeight = height; }
-	void setWalkableRadius(dtReal radius) { m_walkableRadius = radius; }
-	void setWalkableClimb(dtReal climb) { m_walkableClimb = climb; }
-	void setBVQuantFactor(dtReal factor) { m_bvQuantFactor = factor; }
 	//@UE END LWCoords
 
 private:
@@ -924,10 +926,10 @@ private:
 	dtReal m_tileWidth, m_tileHeight;	///< Dimensions of each tile.
 
 	//@UE BEGIN Memory optimization
-	dtReal m_walkableHeight;				///< The height of the agents using the tile.
-	dtReal m_walkableRadius;				///< The radius of the agents using the tile.
+	dtReal m_walkableHeight;			///< The height of the agents using the tile.
+	dtReal m_walkableRadius;			///< The radius of the agents using the tile.
 	dtReal m_walkableClimb;				///< The maximum climb height of the agents using the tile.
-	dtReal m_bvQuantFactor;				///<  The bounding volume quantization factor. 
+	dtReal m_bvQuantFactor;				///< The bounding volume quantization factor. 
 	//@UE END Memory optimization
 
 	int m_maxTiles;						///< Max number of tiles.
