Add way to get random position inside a specified nav polygon

I need that because I do a floodfill query (with ARecastNavMesh::FindPolysAroundCircle) and then need to test positions to move in the resulting polygons

diff --git a/Engine/Source/Runtime/NavigationSystem/Private/NavMesh/PImplRecastNavMesh.cpp b/Engine/Source/Runtime/NavigationSystem/Private/NavMesh/PImplRecastNavMesh.cpp
--- a/Engine/Source/Runtime/NavigationSystem/Private/NavMesh/PImplRecastNavMesh.cpp
+++ b/Engine/Source/Runtime/NavigationSystem/Private/NavMesh/PImplRecastNavMesh.cpp
@@ -2194,6 +2194,24 @@ bool FPImplRecastNavMesh::GetPolyVerts(NavNodeRef PolyID, TArray<FVector>& OutVe
 	return false;
 }
 
+bool FPImplRecastNavMesh::GetRandomPointInPoly(NavNodeRef PolyID, FVector& OutPoint) const
+{
+	if (DetourNavMesh)
+	{
+		INITIALIZE_NAVQUERY_SIMPLE(NavQuery, RECAST_MAX_SEARCH_NODES);
+
+		FVector::FReal RandPt[3];
+		dtStatus Status = NavQuery.findRandomPointInPoly((dtPolyRef)PolyID, FMath::FRand, RandPt);
+		if (dtStatusSucceed(Status))
+		{
+			OutPoint = Recast2UnrVector(RandPt);
+			return true;
+		}
+	}
+
+	return false;
+}
+
 uint32 FPImplRecastNavMesh::GetPolyAreaID(NavNodeRef PolyID) const
 {
 	uint32 AreaID = RECAST_NULL_AREA;
diff --git a/Engine/Source/Runtime/NavigationSystem/Private/NavMesh/RecastNavMesh.cpp b/Engine/Source/Runtime/NavigationSystem/Private/NavMesh/RecastNavMesh.cpp
--- a/Engine/Source/Runtime/NavigationSystem/Private/NavMesh/RecastNavMesh.cpp
+++ b/Engine/Source/Runtime/NavigationSystem/Private/NavMesh/RecastNavMesh.cpp
@@ -1720,6 +1720,11 @@ bool ARecastNavMesh::GetPolyVerts(NavNodeRef PolyID, TArray<FVector>& OutVerts)
 	return RecastNavMeshImpl && RecastNavMeshImpl->GetPolyVerts(PolyID, OutVerts);
 }
 
+bool ARecastNavMesh::GetRandomPointInPoly(NavNodeRef PolyID, FVector& OutPoint) const
+{
+	return RecastNavMeshImpl && RecastNavMeshImpl->GetRandomPointInPoly(PolyID, OutPoint);
+}
+
 uint32 ARecastNavMesh::GetPolyAreaID(NavNodeRef PolyID) const
 {
 	uint32 AreaID = RECAST_DEFAULT_AREA;
diff --git a/Engine/Source/Runtime/NavigationSystem/Public/NavMesh/PImplRecastNavMesh.h b/Engine/Source/Runtime/NavigationSystem/Public/NavMesh/PImplRecastNavMesh.h
--- a/Engine/Source/Runtime/NavigationSystem/Public/NavMesh/PImplRecastNavMesh.h
+++ b/Engine/Source/Runtime/NavigationSystem/Public/NavMesh/PImplRecastNavMesh.h
@@ -159,6 +159,8 @@ public:
 	bool GetPolyCenter(NavNodeRef PolyID, FVector& OutCenter) const;
 	/** Retrieves the vertices for the specified polygon. Returns false on error. */
 	bool GetPolyVerts(NavNodeRef PolyID, TArray<FVector>& OutVerts) const;
+	/** Retrieves a random point inside the specified polygon. Returns false on error. */
+	bool GetRandomPointInPoly(NavNodeRef PolyID, FVector& OutPoint) const;
 	/** Retrieves the flags for the specified polygon. Returns false on error. */
 	bool GetPolyData(NavNodeRef PolyID, uint16& Flags, uint8& AreaType) const;
 	/** Retrieves area ID for the specified polygon. */
diff --git a/Engine/Source/Runtime/NavigationSystem/Public/NavMesh/RecastNavMesh.h b/Engine/Source/Runtime/NavigationSystem/Public/NavMesh/RecastNavMesh.h
--- a/Engine/Source/Runtime/NavigationSystem/Public/NavMesh/RecastNavMesh.h
+++ b/Engine/Source/Runtime/NavigationSystem/Public/NavMesh/RecastNavMesh.h
@@ -1103,6 +1103,9 @@ public:
 	/** Retrieves the vertices for the specified polygon. Returns false on error. */
 	bool GetPolyVerts(NavNodeRef PolyID, TArray<FVector>& OutVerts) const;
 
+	/** Retrieves a random point inside the specified polygon. Returns false on error. */
+	bool GetRandomPointInPoly(NavNodeRef PolyID, FVector& OutPoint) const;
+
 	/** Retrieves area ID for the specified polygon. */
 	uint32 GetPolyAreaID(NavNodeRef PolyID) const;
 
diff --git a/Engine/Source/Runtime/Navmesh/Private/Detour/DetourNavMeshQuery.cpp b/Engine/Source/Runtime/Navmesh/Private/Detour/DetourNavMeshQuery.cpp
--- a/Engine/Source/Runtime/Navmesh/Private/Detour/DetourNavMeshQuery.cpp
+++ b/Engine/Source/Runtime/Navmesh/Private/Detour/DetourNavMeshQuery.cpp
@@ -372,13 +372,11 @@ dtStatus dtNavMeshQuery::findRandomPoint(const dtQueryFilter* filter, float(*fra
 		return DT_FAILURE;
 
 	// Randomly pick point on polygon.
-	const dtReal* v = &tile->verts[poly->verts[0]*3];
 	dtReal verts[3*DT_VERTS_PER_POLYGON];
 	dtReal areas[DT_VERTS_PER_POLYGON];
-	dtVcopy(&verts[0*3],v);
-	for (int j = 1; j < poly->vertCount; ++j)
+	for (int j = 0; j < poly->vertCount; ++j)
 	{
-		v = &tile->verts[poly->verts[j]*3];
+		const dtReal* v = &tile->verts[poly->verts[j]*3];
 		dtVcopy(&verts[j*3],v);
 	}
 	
@@ -570,11 +568,9 @@ dtStatus dtNavMeshQuery::findRandomPointAroundCircle(dtPolyRef startRef, const d
 		const dtPoly* testPoly = 0;
 		m_nav->getTileAndPolyByRefUnsafe(testRef, &testTile, &testPoly);
 
-		const dtReal* v = &testTile->verts[testPoly->verts[0]*3];
-		dtVcopy(&verts[0*3],v);
-		for (int j = 1; j < testPoly->vertCount; ++j)
+		for (int j = 0; j < testPoly->vertCount; ++j)
 		{
-			v = &testTile->verts[testPoly->verts[j]*3];
+			const dtReal* v = &testTile->verts[testPoly->verts[j]*3];
 			dtVcopy(&verts[j*3],v);
 		}
 
@@ -605,6 +601,40 @@ dtStatus dtNavMeshQuery::findRandomPointAroundCircle(dtPolyRef startRef, const d
 	return foundPt ? DT_SUCCESS : DT_FAILURE;
 }
 
+dtStatus dtNavMeshQuery::findRandomPointInPoly(dtPolyRef ref, float(*frand)(), dtReal* randomPt) const
+{
+	dtAssert(m_nav);
+
+	const dtMeshTile* tile = 0;
+	const dtPoly* poly = 0;
+	if (dtStatusFailed(m_nav->getTileAndPolyByRef(ref, &tile, &poly)))
+		return DT_FAILURE | DT_INVALID_PARAM;
+
+	dtReal verts[3*DT_VERTS_PER_POLYGON] = {0};
+	dtReal areas[DT_VERTS_PER_POLYGON] = {0};
+	for (int j = 0; j < poly->vertCount; ++j)
+	{
+		const dtReal* v = &tile->verts[poly->verts[j]*3];
+		dtVcopy(&verts[j*3],v);
+	}
+
+	const dtReal s = frand();
+	const dtReal t = frand();
+
+	dtReal pt[3];
+	dtRandomPointInConvexPoly(verts, poly->vertCount, areas, s, t, pt);
+
+	dtReal h = 0.0;
+	dtStatus status = getPolyHeight(ref, pt, &h);
+	if (dtStatusFailed(status))
+		return status;
+	pt[1] = h;
+
+	dtVcopy(randomPt, pt);
+
+	return DT_SUCCESS;
+}
+
 //@UE BEGIN
 #if WITH_NAVMESH_CLUSTER_LINKS
 dtStatus dtNavMeshQuery::findRandomPointInCluster(dtClusterRef clusterRef, float(*frand)(), dtPolyRef* randomRef, dtReal* randomPt) const
@@ -663,13 +693,11 @@ dtStatus dtNavMeshQuery::findRandomPointInCluster(dtClusterRef clusterRef, float
 	dtPolyRef randomPolyRef = m_nav->getPolyRefBase(searchTile) | (dtPolyRef)randomPolyIdx;
 
 	// Randomly pick point on polygon.
-	const dtReal* v = &searchTile->verts[randomPoly->verts[0]*3];
 	dtReal verts[3*DT_VERTS_PER_POLYGON];
 	dtReal areas[DT_VERTS_PER_POLYGON];
-	dtVcopy(&verts[0*3],v);
-	for (int j = 1; j < randomPoly->vertCount; ++j)
+	for (int j = 0; j < randomPoly->vertCount; ++j)
 	{
-		v = &searchTile->verts[randomPoly->verts[j]*3];
+		const dtReal* v = &searchTile->verts[randomPoly->verts[j]*3];
 		dtVcopy(&verts[j*3],v);
 	}
 
diff --git a/Engine/Source/Runtime/Navmesh/Public/Detour/DetourNavMeshQuery.h b/Engine/Source/Runtime/Navmesh/Public/Detour/DetourNavMeshQuery.h
--- a/Engine/Source/Runtime/Navmesh/Public/Detour/DetourNavMeshQuery.h
+++ b/Engine/Source/Runtime/Navmesh/Public/Detour/DetourNavMeshQuery.h
@@ -639,6 +639,13 @@ public:
 										 const dtQueryFilter* filter, float(*frand)(),
 										 dtPolyRef* randomRef, dtReal* randomPt) const;
 
+	/// Returns random location inside the specified polygon.
+	///  @param[in]		ref				The reference id of the polygon.
+	///  @param[in]		frand			Function returning a random number [0..1).
+	///  @param[out]	randomPt		The random location. [(x, y, z)]
+	/// @returns The status flags for the query.
+	dtStatus findRandomPointInPoly(dtPolyRef ref, float(*frand)(), dtReal* randomPt) const;
+
 	//@UE BEGIN
 #if WITH_NAVMESH_CLUSTER_LINKS
 	/// Check if there is a path from start polygon to the end polygon using cluster graph
