Fix for dtOffMeshSegmentConnection::height is never serialized. It turns out this variable is never used and the whole off mesh segment connection feature is unsuported and has known issues that were never solved at the conceptual level.

I've added comments to make it clearer this feature is unsupported and removed the height variable that isn't used to simplify and tidy up the existing code.

diff --git a/Engine/Source/Runtime/NavigationSystem/Public/NavMesh/RecastNavMesh.h b/Engine/Source/Runtime/NavigationSystem/Public/NavMesh/RecastNavMesh.h
--- a/Engine/Source/Runtime/NavigationSystem/Public/NavMesh/RecastNavMesh.h
+++ b/Engine/Source/Runtime/NavigationSystem/Public/NavMesh/RecastNavMesh.h
@@ -197,6 +197,7 @@ struct FRecastDebugGeometry
 	};
 #endif // WITH_NAVMESH_CLUSTER_LINKS
 
+// This is an unsupported feature and has not been finished to production quality.
 #if WITH_NAVMESH_SEGMENT_LINKS
 	struct FOffMeshSegment
 	{
diff --git a/Engine/Source/Runtime/Navmesh/Navmesh.Build.cs b/Engine/Source/Runtime/Navmesh/Navmesh.Build.cs
--- a/Engine/Source/Runtime/Navmesh/Navmesh.Build.cs
+++ b/Engine/Source/Runtime/Navmesh/Navmesh.Build.cs
@@ -11,6 +11,7 @@ namespace UnrealBuildTool.Rules
 
             PrivateDependencyModuleNames.AddRange(new string[] { "Core" });
 
+			// This is an unsupported feature and has not been finished to production quality.
 			if (Target.bCompileNavmeshSegmentLinks)
 			{
 				PublicDefinitions.Add("WITH_NAVMESH_SEGMENT_LINKS=1");
diff --git a/Engine/Source/Runtime/Navmesh/Private/Detour/DetourNavMeshBuilder.cpp b/Engine/Source/Runtime/Navmesh/Private/Detour/DetourNavMeshBuilder.cpp
--- a/Engine/Source/Runtime/Navmesh/Private/Detour/DetourNavMeshBuilder.cpp
+++ b/Engine/Source/Runtime/Navmesh/Private/Detour/DetourNavMeshBuilder.cpp
@@ -791,7 +791,6 @@ bool dtCreateNavMeshData(dtNavMeshCreateParams* params, unsigned char** outData,
 				dtVcopy(con->endB, &offMeshCon.vertsB1[0]);
 
 				con->rad = offMeshCon.snapRadius;
-				con->height = offMeshCon.snapHeight;
 				con->setFlags(offMeshCon.type);
 				if (offMeshCon.userID)
 					con->userId = offMeshCon.userID;
diff --git a/Engine/Source/Runtime/Navmesh/Public/Detour/DetourNavMesh.h b/Engine/Source/Runtime/Navmesh/Public/Detour/DetourNavMesh.h
--- a/Engine/Source/Runtime/Navmesh/Public/Detour/DetourNavMesh.h
+++ b/Engine/Source/Runtime/Navmesh/Public/Detour/DetourNavMesh.h
@@ -265,6 +265,7 @@ struct dtBVNode
 };
 
 //@UE BEGIN
+// This is an unsupported feature and has not been finished to production quality.
 #if WITH_NAVMESH_SEGMENT_LINKS
 struct dtOffMeshSegmentConnection
 {
@@ -276,9 +277,6 @@ struct dtOffMeshSegmentConnection
 	/// The radius of the endpoints. [Limit: >= 0]
 	dtReal rad;
 
-	/// The snap height of endpoints (less than 0 = use step height)
-	dtReal height;
-
 	/// The id of the offmesh connection. (User assigned when the navigation mesh is built.)
 	unsigned int userId;	
 	
