[Navmesh] Update dtClosestHeightPointTriangle() to check denominator (plus bonus optimization and improved precision)

diff --git a/Engine/Source/Runtime/Navmesh/Private/Detour/DetourCommon.cpp b/Engine/Source/Runtime/Navmesh/Private/Detour/DetourCommon.cpp
--- a/Engine/Source/Runtime/Navmesh/Private/Detour/DetourCommon.cpp
+++ b/Engine/Source/Runtime/Navmesh/Private/Detour/DetourCommon.cpp
@@ -232,30 +232,28 @@ void dtCalcPolyCenter(dtReal* tc, const unsigned short* idx, int nidx, const dtR
 
 bool dtClosestHeightPointTriangle(const dtReal* p, const dtReal* a, const dtReal* b, const dtReal* c, dtReal& h)
 {
-	dtReal v0[3], v1[3], v2[3];
-	dtVsub(v0, c,a);
-	dtVsub(v1, b,a);
-	dtVsub(v2, p,a);
-	
-	const dtReal dot00 = dtVdot2D(v0, v0);
-	const dtReal dot01 = dtVdot2D(v0, v1);
-	const dtReal dot02 = dtVdot2D(v0, v2);
-	const dtReal dot11 = dtVdot2D(v1, v1);
-	const dtReal dot12 = dtVdot2D(v1, v2);
-	
-	// Compute barycentric coordinates
-	const dtReal invDenom = 1.0f / (dot00 * dot11 - dot01 * dot01);
-	const dtReal u = (dot11 * dot02 - dot01 * dot12) * invDenom;
-	const dtReal v = (dot00 * dot12 - dot01 * dot02) * invDenom;
+	dtReal vC[3], vB[3], vP[3];
+	dtVsub(vC, c, a);
+	dtVsub(vB, b, a);
+	dtVsub(vP, p, a);
+
+	// Compute scaled barycentric coordinates
+	dtReal denom = vC[0] * vB[2] - vC[2] * vB[0];
 
 	// The (sloppy) epsilon is needed to allow to get height of points which
 	// are interpolated along the edges of the triangles.
-	static const dtReal EPS = 1e-4f;
-	
+	constexpr dtReal EPS = 1.0e-6;
+	if (fabs(denom) < EPS)
+		return false;
+
+	const dtReal invDenom = 1.0 / denom;
+	const dtReal u = (vB[2] * vP[0] - vB[0] * vP[2]) * invDenom;
+	const dtReal v = (vC[0] * vP[2] - vC[2] * vP[0]) * invDenom;
+
 	// If point lies inside the triangle, return interpolated ycoord.
-	if (u >= -EPS && v >= -EPS && (u+v) <= 1+EPS)
+	if (u >= -EPS && v >= -EPS && (u+v) <= 1.0+EPS)
 	{
-		h = a[1] + v0[1]*u + v1[1]*v;
+		h = a[1] + (vC[1]*u + vB[1]*v);
 		return true;
 	}
 	
