Fix RandomPoly crash with 2 points poly :
- add asserts in dtRandomPointInConvexPoly to detect problematic cases
- add special case for point to point links in dtNavMeshQuery::findRandomPointInPoly

diff --git a/Engine/Source/Runtime/Navmesh/Private/Detour/DetourCommon.cpp b/Engine/Source/Runtime/Navmesh/Private/Detour/DetourCommon.cpp
--- a/Engine/Source/Runtime/Navmesh/Private/Detour/DetourCommon.cpp
+++ b/Engine/Source/Runtime/Navmesh/Private/Detour/DetourCommon.cpp
@@ -20,6 +20,7 @@
 //
 
 #include "Detour/DetourCommon.h"
+#include "Detour/DetourAssert.h"
 
 //////////////////////////////////////////////////////////////////////////////////////////
 
@@ -360,6 +361,8 @@ bool dtOverlapPolyPoly2D(const dtReal* polya, const int npolya,
 void dtRandomPointInConvexPoly(const dtReal* pts, const int npts, dtReal* areas,
 							   const dtReal s, const dtReal t, dtReal* out)
 {
+	dtAssert(npts > 2);
+
 	// Calc triangle araes
 	dtReal areasum = 0.0f;
 	for (int i = 2; i < npts; i++) {
@@ -381,6 +384,8 @@ void dtRandomPointInConvexPoly(const dtReal* pts, const int npts, dtReal* areas,
 		}
 		acc += dacc;
 	}
+
+	dtAssert(tri >= 2);
 	
 	dtReal v = dtSqrt(t);
 	
diff --git a/Engine/Source/Runtime/Navmesh/Private/Detour/DetourNavMeshQuery.cpp b/Engine/Source/Runtime/Navmesh/Private/Detour/DetourNavMeshQuery.cpp
--- a/Engine/Source/Runtime/Navmesh/Private/Detour/DetourNavMeshQuery.cpp
+++ b/Engine/Source/Runtime/Navmesh/Private/Detour/DetourNavMeshQuery.cpp
@@ -610,27 +610,38 @@ dtStatus dtNavMeshQuery::findRandomPointInPoly(dtPolyRef ref, float(*frand)(), d
 	if (dtStatusFailed(m_nav->getTileAndPolyByRef(ref, &tile, &poly)))
 		return DT_FAILURE | DT_INVALID_PARAM;
 
-	dtReal verts[3*DT_VERTS_PER_POLYGON] = {0};
-	dtReal areas[DT_VERTS_PER_POLYGON] = {0};
-	for (int j = 0; j < poly->vertCount; ++j)
+	if (poly->getType() == DT_POLYTYPE_OFFMESH_POINT)
 	{
-		const dtReal* v = &tile->verts[poly->verts[j]*3];
-		dtVcopy(&verts[j*3],v);
+		const dtReal* v0 = &tile->verts[poly->verts[0]*3];
+		const dtReal* v1 = &tile->verts[poly->verts[1]*3];
+		const dtReal s = frand();
+		dtVlerp(randomPt, v0, v1, s);
+		return DT_SUCCESS;
 	}
+	else
+	{
+		dtReal verts[3*DT_VERTS_PER_POLYGON] = {0};
+		dtReal areas[DT_VERTS_PER_POLYGON] = {0};
+		for (int j = 0; j < poly->vertCount; ++j)
+		{
+			const dtReal* v = &tile->verts[poly->verts[j]*3];
+			dtVcopy(&verts[j*3],v);
+		}
 
-	const dtReal s = frand();
-	const dtReal t = frand();
+		const dtReal s = frand();
+		const dtReal t = frand();
 
-	dtReal pt[3];
-	dtRandomPointInConvexPoly(verts, poly->vertCount, areas, s, t, pt);
+		dtReal pt[3];
+		dtRandomPointInConvexPoly(verts, poly->vertCount, areas, s, t, pt);
 
-	dtReal h = 0.0;
-	dtStatus status = getPolyHeight(ref, pt, &h);
-	if (dtStatusFailed(status))
-		return status;
-	pt[1] = h;
+		dtReal h = 0.0;
+		dtStatus status = getPolyHeight(ref, pt, &h);
+		if (dtStatusFailed(status))
+			return status;
+		pt[1] = h;
 
-	dtVcopy(randomPt, pt);
+		dtVcopy(randomPt, pt);
+	}
 
 	return DT_SUCCESS;
 }
