[Navmesh] Fix low height span filtering during navmesh generation
This fixes a case of missing navmesh under a giant mushroom

diff --git a/Engine/Source/Runtime/Navmesh/Private/Recast/RecastFilter.cpp b/Engine/Source/Runtime/Navmesh/Private/Recast/RecastFilter.cpp
--- a/Engine/Source/Runtime/Navmesh/Private/Recast/RecastFilter.cpp
+++ b/Engine/Source/Runtime/Navmesh/Private/Recast/RecastFilter.cpp
@@ -198,7 +198,7 @@ void rcFilterWalkableLowHeightSpans(rcContext* ctx, int walkableHeight, rcHeight
 			{
 				const int bot = (int)(s->data.smax);
 				const int top = s->next ? (int)(s->next->data.smin) : MAX_HEIGHT;
-				if ((top - bot) <= walkableHeight)
+				if ((top - bot) < walkableHeight)	// UE
 					s->data.area = RC_NULL_AREA;
 			}
 		}
