RecastNavmesh: Added FindOverlappingEdges() method to get wall edges that intersect a convex polygon.

diff --git a/Engine/Source/Runtime/NavigationSystem/Private/NavMesh/RecastNavMesh.cpp b/Engine/Source/Runtime/NavigationSystem/Private/NavMesh/RecastNavMesh.cpp
--- a/Engine/Source/Runtime/NavigationSystem/Private/NavMesh/RecastNavMesh.cpp
+++ b/Engine/Source/Runtime/NavigationSystem/Private/NavMesh/RecastNavMesh.cpp
@@ -1329,6 +1329,58 @@ bool ARecastNavMesh::FindMoveAlongSurface(const FNavLocation& StartLocation, con
 	return bSuccess;
 }
 
+bool ARecastNavMesh::FindOverlappingEdges(const FNavLocation& StartLocation, TConstArrayView<FVector> ConvexPolygon, TArray<FVector>& OutEdges, FSharedConstNavQueryFilter Filter, const UObject* Querier) const
+{
+	if (RecastNavMeshImpl == NULL || RecastNavMeshImpl->DetourNavMesh == NULL)
+	{
+		return false;
+	}
+
+	const FNavigationQueryFilter& FilterToUse = GetRightFilterRef(Filter);
+
+	FRecastSpeciaLinkFilter LinkFilter(FNavigationSystem::GetCurrent<UNavigationSystemV1>(GetWorld()), Querier);
+	INITIALIZE_NAVQUERY_WLINKFILTER(NavQuery, FilterToUse.GetMaxSearchNodes(), LinkFilter);
+	const dtQueryFilter* QueryFilter = static_cast<const FRecastQueryFilter*>(FilterToUse.GetImplementation())->GetAsDetourQueryFilter();
+
+	const int32 MaxWalls = 64;
+	int32 NumWalls = 0;
+	FVector::FReal WallSegments[MaxWalls * 3 * 2] = { 0 };
+	dtPolyRef WallPolys[MaxWalls * 2] = { 0 };
+
+	const int32 MaxNeis = 64;
+	int32 NumNeis = 0;
+	dtPolyRef NeiPolys[MaxNeis] = { 0 };
+
+	const int32 MaxConvexPolygonPoints = 8;
+	int32 NumConvexPolygonPoints = FMath::Min(ConvexPolygon.Num(), MaxConvexPolygonPoints);
+	FVector::FReal RcConvexPolygon[MaxConvexPolygonPoints * 3] = { 0 };
+
+	for (int32 i  = 0; i < NumConvexPolygonPoints; i++)
+	{
+		const FVector RcPoint = Unreal2RecastPoint(ConvexPolygon[i]);
+		RcConvexPolygon[i*3+0] = RcPoint.X;
+		RcConvexPolygon[i*3+1] = RcPoint.Y;
+		RcConvexPolygon[i*3+2] = RcPoint.Z;
+	}
+
+	dtStatus Status = NavQuery.findWallsOverlappingShape(StartLocation.NodeRef, RcConvexPolygon, NumConvexPolygonPoints, QueryFilter,
+		NeiPolys, &NumNeis, MaxNeis, WallSegments, WallPolys, &NumWalls, MaxWalls);
+
+	if (dtStatusSucceed(Status))
+	{
+		OutEdges.Reset(NumWalls*2);
+		for (int32 Idx = 0; Idx < NumWalls; Idx++)
+		{
+			OutEdges.Add(Recast2UnrealPoint(&WallSegments[Idx * 6]));
+			OutEdges.Add(Recast2UnrealPoint(&WallSegments[Idx * 6 + 3]));
+		}
+
+		return true;
+	}
+
+	return false;
+}
+
 bool ARecastNavMesh::ProjectPoint(const FVector& Point, FNavLocation& OutLocation, const FVector& Extent, FSharedConstNavQueryFilter Filter, const UObject* QueryOwner) const
 {
 	bool bSuccess = false;
diff --git a/Engine/Source/Runtime/NavigationSystem/Public/NavMesh/RecastNavMesh.h b/Engine/Source/Runtime/NavigationSystem/Public/NavMesh/RecastNavMesh.h
--- a/Engine/Source/Runtime/NavigationSystem/Public/NavMesh/RecastNavMesh.h
+++ b/Engine/Source/Runtime/NavigationSystem/Public/NavMesh/RecastNavMesh.h
@@ -892,6 +892,7 @@ public:
 	virtual bool GetRandomPointInNavigableRadius(const FVector& Origin, float Radius, FNavLocation& OutResult, FSharedConstNavQueryFilter Filter = NULL, const UObject* Querier = NULL) const override;
 
 	virtual bool FindMoveAlongSurface(const FNavLocation& StartLocation, const FVector& TargetPosition, FNavLocation& OutLocation, FSharedConstNavQueryFilter Filter = NULL, const UObject* Querier = NULL) const override;
+	virtual bool FindOverlappingEdges(const FNavLocation& StartLocation, TConstArrayView<FVector> ConvexPolygon, TArray<FVector>& OutEdges, FSharedConstNavQueryFilter Filter = NULL, const UObject* Querier = NULL) const override;
 	virtual bool ProjectPoint(const FVector& Point, FNavLocation& OutLocation, const FVector& Extent, FSharedConstNavQueryFilter Filter = NULL, const UObject* Querier = NULL) const override;
 	virtual bool IsNodeRefValid(NavNodeRef NodeRef) const override;
 
diff --git a/Engine/Source/Runtime/NavigationSystem/Public/NavigationData.h b/Engine/Source/Runtime/NavigationSystem/Public/NavigationData.h
--- a/Engine/Source/Runtime/NavigationSystem/Public/NavigationData.h
+++ b/Engine/Source/Runtime/NavigationSystem/Public/NavigationData.h
@@ -833,6 +833,8 @@ public:
 	 */
 	virtual bool FindMoveAlongSurface(const FNavLocation& StartLocation, const FVector& TargetPosition, FNavLocation& OutLocation, FSharedConstNavQueryFilter Filter = NULL, const UObject* Querier = NULL) const PURE_VIRTUAL(ANavigationData::FindMoveAlongSurface, return false;);
 
+	virtual bool FindOverlappingEdges(const FNavLocation& StartLocation, TConstArrayView<FVector> ConvexPolygon, TArray<FVector>& OutEdges, FSharedConstNavQueryFilter Filter = NULL, const UObject* Querier = NULL) const PURE_VIRTUAL(ANavigationData::FindOverlappingEdges, return false;);
+
 	virtual FNavLocation GetRandomPoint(FSharedConstNavQueryFilter Filter = NULL, const UObject* Querier = NULL) const PURE_VIRTUAL(ANavigationData::GetRandomPoint, return FNavLocation(););
 
 	/** finds a random location in Radius, reachable from Origin */
diff --git a/Engine/Source/Runtime/Navmesh/Private/Detour/DetourNavMeshQuery.cpp b/Engine/Source/Runtime/Navmesh/Private/Detour/DetourNavMeshQuery.cpp
--- a/Engine/Source/Runtime/Navmesh/Private/Detour/DetourNavMeshQuery.cpp
+++ b/Engine/Source/Runtime/Navmesh/Private/Detour/DetourNavMeshQuery.cpp
@@ -4306,101 +4306,287 @@ dtStatus dtNavMeshQuery::getPolyWallSegments(dtPolyRef ref, const dtQueryFilter*
 	return status;
 }
 
-static void storeWallSegment(const dtMeshTile* tile, const dtPoly* poly, int edge,
-	dtPolyRef ref0, dtPolyRef ref1, const dtNavMesh* nav, const dtReal* centerPos, const dtReal radiusSqr,
-	dtReal* resultWalls, dtPolyRef* resultRefs, int* resultCount, const int maxResult)
+static bool isEdgeInRadius(const dtMeshTile* tile, const dtPoly* poly, int edge, const dtReal* centerPos, const dtReal radiusSqr)
 {
-	if (*resultCount < maxResult)
+	const dtReal* va = &tile->verts[poly->verts[edge] * 3];
+	const dtReal* vb = &tile->verts[poly->verts[(edge + 1) % poly->vertCount] * 3];
+
+	dtReal tseg;
+	const dtReal distSqr = dtDistancePtSegSqr2D(centerPos, va, vb, tseg);
+	return distSqr <= radiusSqr;
+}
+
+static void storeWallSegment(const dtNavMesh* nav, const dtMeshTile* tile, const dtPoly* poly, int edge, dtPolyRef ref0, dtPolyRef ref1,
+							 dtReal* resultWalls, dtPolyRef* resultRefs, int* resultCount, const int maxResult)
+{
+	if (*resultCount >= maxResult)
 	{
-		const dtReal* va = &tile->verts[poly->verts[edge] * 3];
-		const dtReal* vb = &tile->verts[poly->verts[(edge + 1) % poly->vertCount] * 3];
+		return;
+	}
+	
+	const dtReal* va = &tile->verts[poly->verts[edge] * 3];
+	const dtReal* vb = &tile->verts[poly->verts[(edge + 1) % poly->vertCount] * 3];
 
-		dtReal tseg;
-		dtReal distSqr = dtDistancePtSegSqr2D(centerPos, va, vb, tseg);
-		if (distSqr <= radiusSqr)
+	const int32 wall0Offset = (*resultCount * 6) + 0;
+	const int32 wall1Offset = (*resultCount * 6) + 3;
+
+	dtVcopy(&resultWalls[wall0Offset], va);
+	dtVcopy(&resultWalls[wall1Offset], vb);
+	resultRefs[*resultCount * 2 + 0] = ref0;
+	resultRefs[*resultCount * 2 + 1] = ref1;
+
+	*resultCount += 1;
+
+	// If neighbour is valid, find the segment that both polygons share by projecting the neighbour segment to the current segment va-vb.
+	if (ref1)
+	{
+		const dtMeshTile* neiTile = 0;
+		const dtPoly* neiPoly = 0;
+		nav->getTileAndPolyByRef(ref1, &neiTile, &neiPoly);
+
+		// Find edge of the neighbour polygon.
+		int neiEdge = -1;
+		unsigned int neiLinkId = neiPoly ? neiPoly->firstLink : DT_NULL_LINK;
+		while (neiLinkId != DT_NULL_LINK)
 		{
-			const int32 Wall0Offset = (*resultCount * 6) + 0;
-			const int32 Wall1Offset = (*resultCount * 6) + 3;
+			const dtLink& link = nav->getLink(neiTile, neiLinkId);
+			neiLinkId = link.next;
+			if (link.ref == ref0)
+			{
+				neiEdge = link.edge;
+				break;
+			}
+		}
+
+		if (neiEdge != -1)
+		{
+			const dtReal* va2 = &neiTile->verts[neiPoly->verts[neiEdge] * 3];
+			const dtReal* vb2 = &neiTile->verts[neiPoly->verts[(neiEdge + 1) % neiPoly->vertCount] * 3];
 
-			dtVcopy(&resultWalls[Wall0Offset], va);
-			dtVcopy(&resultWalls[Wall1Offset], vb);
-			resultRefs[*resultCount * 2 + 0] = ref0;
-			resultRefs[*resultCount * 2 + 1] = ref1;
+			// project segment va2-vb2 on va-vb: point va2
+			dtReal seg[3], toPt[3], closestA[3], closestB[3];
+			dtVsub(seg, vb, va);
 
-			*resultCount += 1;
+			dtVsub(toPt, va2, va);
+			dtReal d1 = dtVdot(toPt, seg);
+			dtReal d2 = dtVdot(seg, seg);
 
-			// find intersection between two edge segments
-			if (nav)
+			if (d1 <= 0)
+			{
+				dtVcopy(closestA, va);
+			}
+			else if (d2 <= d1)
+			{
+				dtVcopy(closestA, vb);
+			}
+			else
+			{
+				dtVmad(closestA, va, seg, d1 / d2);
+			}
+
+			// project segment va2-vb2 on va-vb: point vb2
+			dtVsub(toPt, vb2, va);
+			d1 = dtVdot(toPt, seg);
+			d2 = dtVdot(seg, seg);
+
+			if (d1 <= 0)
+			{
+				dtVcopy(closestB, va);
+			}
+			else if (d2 <= d1)
+			{
+				dtVcopy(closestB, vb);
+			}
+			else
 			{
-				const dtMeshTile* NeiTile = 0;
-				const dtPoly* NeiPoly = 0;
-				nav->getTileAndPolyByRef(ref1, &NeiTile, &NeiPoly);
+				dtVmad(closestB, va, seg, d1 / d2);
+			}
+
+			// store projected segment (intersection of both edges)
+			dtVcopy(&resultWalls[wall0Offset], closestA);
+			dtVcopy(&resultWalls[wall1Offset], closestB);
+		}
+	}
+}
+
+dtStatus dtNavMeshQuery::findWallsInNeighbourhood(dtPolyRef startRef, const dtReal* centerPos, const dtReal radius,
+												  const dtQueryFilter* filter,
+												  dtPolyRef* neiRefs, int* neiCount, const int maxNei,
+												  dtReal* resultWalls, dtPolyRef* resultRefs, int* resultCount, const int maxResult) const
+{
+	*resultCount = 0;
+	*neiCount = 0;
+
+	// Validate input
+	if (!startRef || !m_nav->isValidPolyRef(startRef))
+		return DT_FAILURE | DT_INVALID_PARAM;
+
+	m_tinyNodePool->clear();
+
+	static const int MAX_STACK = 48;
+	dtNode* stack[MAX_STACK];
+	int nstack = 0;
 
-				unsigned int NeiLinkId = NeiPoly ? NeiPoly->firstLink : DT_NULL_LINK;
-				while (NeiLinkId != DT_NULL_LINK)
+	dtNode* startNode = m_tinyNodePool->getNode(startRef);
+	startNode->pidx = 0;
+	startNode->id = startRef;
+	startNode->flags = DT_NODE_CLOSED;
+	stack[nstack++] = startNode;
+
+	dtStatus status = DT_SUCCESS;
+	const dtReal radiusSqr = dtSqr(radius);
+
+	int n = 0;
+	if (n < maxNei)
+	{
+		neiRefs[n] = startNode->id;
+		++n;
+	}
+	else
+	{
+		status |= DT_BUFFER_TOO_SMALL;
+	}
+
+	while (nstack)
+	{
+		// Pop front.
+		dtNode* curNode = stack[0];
+		for (int stackIndex = 0; stackIndex < nstack - 1; ++stackIndex)
+			stack[stackIndex] = stack[stackIndex + 1];
+		nstack--;
+		
+		// Get poly and tile.
+		// The API input has been cheked already, skip checking internal data.
+		const dtPolyRef curRef = curNode->id;
+		const dtMeshTile* curTile = 0;
+		const dtPoly* curPoly = 0;
+		m_nav->getTileAndPolyByRefUnsafe(curRef, &curTile, &curPoly);
+
+		unsigned int i = curPoly->firstLink;
+		while (i != DT_NULL_LINK)
+		{
+			const dtLink& link = m_nav->getLink(curTile, i);
+			i = link.next;
+
+			dtPolyRef neighbourRef = link.ref;
+			// Skip invalid neighbours.
+			if (!neighbourRef)
+			{
+				// store wall segment
+				if (isEdgeInRadius(curTile, curPoly, link.edge, centerPos, radiusSqr))
 				{
-					const dtLink& link = nav->getLink(NeiTile, NeiLinkId);
-					NeiLinkId = link.next;
+					storeWallSegment(m_nav, curTile, curPoly, link.edge,  curRef, 0, 
+						resultWalls, resultRefs, resultCount, maxResult);
+				}
+				continue;
+			}
 
-					if (link.ref == ref0)
-					{
-						const dtReal* va2 = &NeiTile->verts[NeiPoly->verts[link.edge] * 3];
-						const dtReal* vb2 = &NeiTile->verts[NeiPoly->verts[(link.edge + 1) % NeiPoly->vertCount] * 3];
+			// Skip if cannot alloca more nodes.
+			dtNode* neighbourNode = m_tinyNodePool->getNode(neighbourRef);
+			if (!neighbourNode)
+				continue;
+			// Skip visited.
+			if (neighbourNode->flags & DT_NODE_CLOSED)
+				continue;
 
-						// project segment va2-vb2 on va-vb: point va2
-						dtReal seg[3], toPt[3], closestA[3], closestB[3];
-						dtVsub(seg, vb, va);
+			// Expand to neighbour
+			const dtMeshTile* neighbourTile = 0;
+			const dtPoly* neighbourPoly = 0;
+			m_nav->getTileAndPolyByRefUnsafe(neighbourRef, &neighbourTile, &neighbourPoly);
 
-						dtVsub(toPt, va2, va);
-						dtReal d1 = dtVdot(toPt, seg);
-						dtReal d2 = dtVdot(seg, seg);
+			// Skip off-mesh connections.
+			if (neighbourPoly->getType() != DT_POLYTYPE_GROUND)
+				continue;
 
-						if (d1 <= 0)
-						{
-							dtVcopy(closestA, va);
-						}
-						else if (d2 <= d1)
-						{
-							dtVcopy(closestA, vb);
-						}
-						else
-						{
-							dtVmad(closestA, va, seg, d1 / d2);
-						}
+			// Do not advance if the polygon is excluded by the filter.
+			if (!filter->passFilter(neighbourRef, neighbourTile, neighbourPoly) || !passLinkFilterByRef(neighbourTile, neighbourRef))
+			{
+				// store wall segment
+				if (isEdgeInRadius(curTile, curPoly, link.edge, centerPos, radiusSqr))
+				{
+					storeWallSegment(m_nav, curTile, curPoly, link.edge, curRef, neighbourRef,
+						resultWalls, resultRefs, resultCount, maxResult);
+				}
+				continue;
+			}
+
+			// Find edge and calc distance to the edge.
+			dtReal va[3], vb[3];
+			if (!getPortalPoints(curRef, curPoly, curTile, neighbourRef, neighbourPoly, neighbourTile, va, vb))
+				continue;
 
-						// project segment va2-vb2 on va-vb: point vb2
-						dtVsub(toPt, vb2, va);
-						d1 = dtVdot(toPt, seg);
-						d2 = dtVdot(seg, seg);
+			// If the circle is not touching the next polygon, skip it.
+			dtReal tseg;
+			dtReal distSqr = dtDistancePtSegSqr2D(centerPos, va, vb, tseg);
+			if (distSqr > radiusSqr)
+				continue;
 
-						if (d1 <= 0)
-						{
-							dtVcopy(closestB, va);
-						}
-						else if (d2 <= d1)
-						{
-							dtVcopy(closestB, vb);
-						}
-						else
-						{
-							dtVmad(closestB, va, seg, d1 / d2);
-						}
+			// Mark node visited, this is done before the overlap test so that
+			// we will not visit the poly again if the test fails.
+			neighbourNode->flags |= DT_NODE_CLOSED;
+			neighbourNode->pidx = m_tinyNodePool->getNodeIdx(curNode);			
 
-						// store projected segment (intersection of both edges)
-						dtVcopy(&resultWalls[Wall0Offset], closestA);
-						dtVcopy(&resultWalls[Wall1Offset], closestB);
+			// This poly is fine, store and advance to the poly.
+			if (n < maxNei)
+			{
+				neiRefs[n] = neighbourRef;
+				++n;
+			}
+			else
+			{
+				status |= DT_BUFFER_TOO_SMALL;
+			}
+
+			if (nstack < MAX_STACK)
+			{
+				stack[nstack++] = neighbourNode;
+			}
+		}
+
+		// add hard edges of poly
+		for (int neighbourIndex = 0; neighbourIndex < curPoly->vertCount; neighbourIndex++)
+		{
+			bool bStoreEdge = (curPoly->neis[neighbourIndex] == 0);
+			if (curPoly->neis[neighbourIndex] & DT_EXT_LINK)
+			{
+				// check if external edge has valid link
+				bool bConnected = false;
+
+				unsigned int linkIdx = curPoly->firstLink;
+				while (linkIdx != DT_NULL_LINK)
+				{
+					const dtLink& link = m_nav->getLink(curTile, linkIdx);
+					linkIdx = link.next;
+
+					if (link.edge == neighbourIndex)
+					{
+						bConnected = true;
 						break;
 					}
 				}
+
+				bStoreEdge = !bConnected;
+			}
+
+			if (bStoreEdge)
+			{
+				if (isEdgeInRadius(curTile, curPoly, neighbourIndex, centerPos, radiusSqr))
+				{
+					storeWallSegment(m_nav, curTile, curPoly, neighbourIndex, curRef, 0, 
+						resultWalls, resultRefs, resultCount, maxResult);
+				}
 			}
 		}
 	}
+
+	*neiCount = n;
+	return status;
 }
 
-dtStatus dtNavMeshQuery::findWallsInNeighbourhood(dtPolyRef startRef, const dtReal* centerPos, const dtReal radius,
-												  const dtQueryFilter* filter,
-												  dtPolyRef* neiRefs, int* neiCount, const int maxNei,
-												  dtReal* resultWalls, dtPolyRef* resultRefs, int* resultCount, const int maxResult) const
+dtStatus dtNavMeshQuery::findWallsOverlappingShape(dtPolyRef startRef, const dtReal* verts, const int nverts,
+												   const dtQueryFilter* filter,
+												   dtPolyRef* neiRefs, int* neiCount, const int maxNei,
+												   dtReal* resultWalls, dtPolyRef* resultRefs, int* resultCount, const int maxResult) const
 {
 	*resultCount = 0;
 	*neiCount = 0;
@@ -4422,7 +4608,6 @@ dtStatus dtNavMeshQuery::findWallsInNeighbourhood(dtPolyRef startRef, const dtRe
 	stack[nstack++] = startNode;
 
 	dtStatus status = DT_SUCCESS;
-	const dtReal radiusSqr = dtSqr(radius);
 
 	int n = 0;
 	if (n < maxNei)
@@ -4461,8 +4646,7 @@ dtStatus dtNavMeshQuery::findWallsInNeighbourhood(dtPolyRef startRef, const dtRe
 			if (!neighbourRef)
 			{
 				// store wall segment
-				storeWallSegment(curTile, curPoly, link.edge, 
-					curRef, 0, 0, centerPos, radiusSqr, 
+				storeWallSegment(m_nav, curTile, curPoly, link.edge, curRef, 0, 
 					resultWalls, resultRefs, resultCount, maxResult);
 				continue;
 			}
@@ -4488,8 +4672,7 @@ dtStatus dtNavMeshQuery::findWallsInNeighbourhood(dtPolyRef startRef, const dtRe
 			if (!filter->passFilter(neighbourRef, neighbourTile, neighbourPoly) || !passLinkFilterByRef(neighbourTile, neighbourRef))
 			{
 				// store wall segment
-				storeWallSegment(curTile, curPoly, link.edge,
-					curRef, neighbourRef, m_nav, centerPos, radiusSqr,
+				storeWallSegment(m_nav, curTile, curPoly, link.edge, curRef, neighbourRef,
 					resultWalls, resultRefs, resultCount, maxResult);
 
 				continue;
@@ -4500,10 +4683,12 @@ dtStatus dtNavMeshQuery::findWallsInNeighbourhood(dtPolyRef startRef, const dtRe
 			if (!getPortalPoints(curRef, curPoly, curTile, neighbourRef, neighbourPoly, neighbourTile, va, vb))
 				continue;
 
-			// If the circle is not touching the next polygon, skip it.
-			dtReal tseg;
-			dtReal distSqr = dtDistancePtSegSqr2D(centerPos, va, vb, tseg);
-			if (distSqr > radiusSqr)
+			// If the poly is not touching the edge to the next polygon, skip the connection it.
+			dtReal tmin, tmax;
+			int segMin, segMax;
+			if (!dtIntersectSegmentPoly2D(va, vb, verts, nverts, tmin, tmax, segMin, segMax))
+				continue;
+			if (tmin > 1.0f || tmax < 0.0f)
 				continue;
 
 			// Mark node visited, this is done before the overlap test so that
@@ -4555,8 +4740,7 @@ dtStatus dtNavMeshQuery::findWallsInNeighbourhood(dtPolyRef startRef, const dtRe
 
 			if (bStoreEdge)
 			{
-				storeWallSegment(curTile, curPoly, neighbourIndex,
-					curRef, 0, 0, centerPos, radiusSqr, 
+				storeWallSegment(m_nav, curTile, curPoly, neighbourIndex, curRef, 0, 
 					resultWalls, resultRefs, resultCount, maxResult);
 			}
 		}
@@ -4566,7 +4750,6 @@ dtStatus dtNavMeshQuery::findWallsInNeighbourhood(dtPolyRef startRef, const dtRe
 	return status;
 }
 
-
 /// @par
 ///
 /// @p hitPos is not adjusted using the height detail data.
diff --git a/Engine/Source/Runtime/Navmesh/Public/Detour/DetourNavMeshQuery.h b/Engine/Source/Runtime/Navmesh/Public/Detour/DetourNavMeshQuery.h
--- a/Engine/Source/Runtime/Navmesh/Public/Detour/DetourNavMeshQuery.h
+++ b/Engine/Source/Runtime/Navmesh/Public/Detour/DetourNavMeshQuery.h
@@ -556,6 +556,14 @@ public:
 									  dtPolyRef* neiRefs, int* neiCount, const int maxNei,
 									  dtReal* resultWalls, dtPolyRef* resultRefs, int* resultCount, const int maxResult) const;
 
+
+	/// [UE] Finds the wall segments that overlap the polygon shape.
+	dtStatus findWallsOverlappingShape(dtPolyRef startRef, const dtReal* verts, const int nverts,
+									   const dtQueryFilter* filter, 
+									   dtPolyRef* neiRefs, int* neiCount, const int maxNei,
+									   dtReal* resultWalls, dtPolyRef* resultRefs, int* resultCount, const int maxResult) const;
+
+	
 	/// Moves from the start to the end position constrained to the navigation mesh.
 	///  @param[in]		startRef		The reference id of the start polygon.
 	///  @param[in]		startPos		A position of the mover within the start polygon. [(x, y, x)]
