PR #9506: Apply bug fix: "Fix 'dtRandomPointInConvexPoly()' returns a garbage value if s == 1.0f." (Contributed by [REDACTED_PERSON])

diff --git a/Engine/Source/Runtime/Navmesh/Private/Detour/DetourCommon.cpp b/Engine/Source/Runtime/Navmesh/Private/Detour/DetourCommon.cpp
--- a/Engine/Source/Runtime/Navmesh/Private/Detour/DetourCommon.cpp
+++ b/Engine/Source/Runtime/Navmesh/Private/Detour/DetourCommon.cpp
@@ -372,8 +372,8 @@ void dtRandomPointInConvexPoly(const dtReal* pts, const int npts, dtReal* areas,
 	// Find sub triangle weighted by area.
 	const dtReal thr = s*areasum;
 	dtReal acc = 0.0f;
-	dtReal u = 0.0f;
-	int tri = 0;
+	dtReal u = 1.0f;
+	int tri = npts - 1;
 	for (int i = 2; i < npts; i++) {
 		const dtReal dacc = areas[i];
 		if (thr >= acc && thr < (acc+dacc))
