Remove Recast steep slope filtering when the agent radius is only 1 voxel

diff --git a/Engine/Source/Runtime/NavigationSystem/Private/NavMesh/RecastNavMeshGenerator.cpp b/Engine/Source/Runtime/NavigationSystem/Private/NavMesh/RecastNavMeshGenerator.cpp
--- a/Engine/Source/Runtime/NavigationSystem/Private/NavMesh/RecastNavMeshGenerator.cpp
+++ b/Engine/Source/Runtime/NavigationSystem/Private/NavMesh/RecastNavMeshGenerator.cpp
@@ -82,6 +82,9 @@ namespace UE::NavMesh::Private
 
 	static bool bUseTightBoundExpansion = true;
 	static FAutoConsoleVariableRef CVarUseTightBoundExpansion(TEXT("n.UseTightBoundExpansion"), bUseTightBoundExpansion, TEXT("Active by default. Use an expansion of one AgentRadius. Set to false to revert to the previous behavior (2 AgentRadius)."), ECVF_Default);
+
+	static bool bKeepSteepSlopeForSingleVoxelAgent = true;
+	static FAutoConsoleVariableRef CVarKeepSteepSlopeForSingleVoxelAgent(TEXT("n.KeepSteepSlopeForSingleVoxelAgent"), bKeepSteepSlopeForSingleVoxelAgent, TEXT("Active by default. Fix too wide filtering of steep slope in Recast Heightfield filtering when the agent radius is only 1 voxel. Set to false to revert to the previous behavior."), ECVF_Default);
 }
 
 static FOodleDataCompression::ECompressor GNavmeshTileCacheCompressor = FOodleDataCompression::ECompressor::Mermaid;
@@ -2942,7 +2945,8 @@ void FRecastTileGenerator::GenerateRecastFilter(FNavMeshBuildContext& BuildConte
 	{
 		SCOPE_CYCLE_COUNTER(STAT_Navigation_FilterLedgeSpans)
 
-		rcFilterLedgeSpans(&BuildContext, TileConfig.walkableHeight, TileConfig.walkableClimb, *RasterContext.SolidHF);
+		const bool bFilterNeighborSlope = TileConfig.walkableRadius > 1 || !UE::NavMesh::Private::bKeepSteepSlopeForSingleVoxelAgent;
+		rcFilterLedgeSpans(&BuildContext, TileConfig.walkableHeight, TileConfig.walkableClimb, bFilterNeighborSlope, *RasterContext.SolidHF);
 	}
 	if (!TileConfig.bMarkLowHeightAreas)
 	{
@@ -2988,9 +2992,11 @@ ETimeSliceWorkResult FRecastTileGenerator::GenerateRecastFilterTimeSliced(FNavMe
 
 		bool DoIter = true;
 
+		const bool bFilterNeighborSlope = TileConfig.walkableRadius > 1 || !UE::NavMesh::Private::bKeepSteepSlopeForSingleVoxelAgent;
+
 		do
 		{
-			rcFilterLedgeSpans(&BuildContext, TileConfig.walkableHeight, TileConfig.walkableClimb, GenRecastFilterLedgeSpansYStart, TileTimeSliceSettings.FilterLedgeSpansMaxYProcess, *RasterContext.SolidHF);
+			rcFilterLedgeSpans(&BuildContext, TileConfig.walkableHeight, TileConfig.walkableClimb, bFilterNeighborSlope, GenRecastFilterLedgeSpansYStart, TileTimeSliceSettings.FilterLedgeSpansMaxYProcess, *RasterContext.SolidHF);
 
 			GenRecastFilterLedgeSpansYStart += TileTimeSliceSettings.FilterLedgeSpansMaxYProcess;
 
diff --git a/Engine/Source/Runtime/NavigationSystem/Public/NavMesh/RecastVersion.h b/Engine/Source/Runtime/NavigationSystem/Public/NavMesh/RecastVersion.h
--- a/Engine/Source/Runtime/NavigationSystem/Public/NavMesh/RecastVersion.h
+++ b/Engine/Source/Runtime/NavigationSystem/Public/NavMesh/RecastVersion.h
@@ -23,6 +23,7 @@
 #define NAVMESHVER_MAXTILES_COUNT_SKIP_INCLUSION 22
 #define NAVMESHVER_TILE_RESOLUTIONS				23 // Addition of a tile resolution index to the tile header.
 #define NAVMESHVER_TILE_RESOLUTIONS_CELLHEIGHT	24 // Addition of CellHeight in the resolution params, deprecating the original CellHeight.
+#define NAVMESHVER_1_VOXEL_AGENT_STEEP_SLOPE_FILTER_FIX	25 // Fix, remove steep slope filtering during heightfield ledge filtering when the agent radius is included into a single voxel 
 
-#define NAVMESHVER_LATEST				NAVMESHVER_TILE_RESOLUTIONS_CELLHEIGHT
+#define NAVMESHVER_LATEST				NAVMESHVER_1_VOXEL_AGENT_STEEP_SLOPE_FILTER_FIX
 #define NAVMESHVER_MIN_COMPATIBLE		NAVMESHVER_LWCOORDS_OPTIMIZATION
diff --git a/Engine/Source/Runtime/Navmesh/Private/Recast/RecastFilter.cpp b/Engine/Source/Runtime/Navmesh/Private/Recast/RecastFilter.cpp
--- a/Engine/Source/Runtime/Navmesh/Private/Recast/RecastFilter.cpp
+++ b/Engine/Source/Runtime/Navmesh/Private/Recast/RecastFilter.cpp
@@ -72,7 +72,7 @@ void rcFilterLowHangingWalkableObstacles(rcContext* ctx, const int walkableClimb
 	ctx->stopTimer(RC_TIMER_FILTER_LOW_OBSTACLES);
 }
 
-void rcFilterLedgeSpansImp(rcContext* ctx, const int walkableHeight, const int walkableClimb, const int filterLedgeSpansAtY,
+void rcFilterLedgeSpansImp(rcContext* ctx, const int walkableHeight, const int walkableClimb, const bool filterNeighborSlope, const int filterLedgeSpansAtY, //UE
 	rcHeightfield& solid)
 {
 	rcAssert(ctx);
@@ -130,7 +130,7 @@ void rcFilterLedgeSpansImp(rcContext* ctx, const int walkableHeight, const int w
 						minh = rcMin(minh, nbot - bot);
 
 						// Find min/max accessible neighbour height. 
-						if (rcAbs(nbot - bot) <= walkableClimb)
+						if (filterNeighborSlope && rcAbs(nbot - bot) <= walkableClimb)	//UE
 						{
 							if (nbot < asmin) asmin = nbot;
 							if (nbot > asmax) asmax = nbot;
@@ -143,11 +143,20 @@ void rcFilterLedgeSpansImp(rcContext* ctx, const int walkableHeight, const int w
 			// The current span is close to a ledge if the drop to any
 			// neighbour span is less than the walkableClimb.
 			if (minh < -walkableClimb)
+			{
 				s->data.area = RC_NULL_AREA;
-
+			}
 			// If the difference between all neighbours is too large,
 			// we are at steep slope, mark the span as ledge.
-			if ((asmax - asmin) > walkableClimb)
+//@UE BEGIN
+			// This test is meant to reject step/slope combinations that can't be climbed by the agent.
+			// At the span level, we detect the walkable climb height between neighbor spans. At the agent level,
+			// the radius of the agent can mean that we should consider multiple neighbor spans to detect what would be
+			// the actual height of the step to climb. Because of that, depending on the number of voxel represented by
+			// the agent radius, the following test might cause more problems than improvements.
+			// That's why we decided to allow to control it via the filterNeighborSlope parameter
+			else if (filterNeighborSlope && (asmax - asmin) > walkableClimb)
+//@UE END
 			{
 				s->data.area = RC_NULL_AREA;
 			}
@@ -165,7 +174,7 @@ void rcFilterLedgeSpansImp(rcContext* ctx, const int walkableHeight, const int w
 /// A span is a ledge if: <tt>rcAbs(currentSpan.smax - neighborSpan.smax) > walkableClimb</tt>
 /// 
 /// @see rcHeightfield, rcConfig
-void rcFilterLedgeSpans(rcContext* ctx, const int walkableHeight, const int walkableClimb,
+void rcFilterLedgeSpans(rcContext* ctx, const int walkableHeight, const int walkableClimb, const bool filterNeighborSlope, //UE
 						rcHeightfield& solid)
 {
 	rcAssert(ctx);
@@ -179,15 +188,15 @@ void rcFilterLedgeSpans(rcContext* ctx, const int walkableHeight, const int walk
 	// Mark border spans.
 	for (int y = 0; y < h; ++y)
 	{
-		rcFilterLedgeSpansImp(ctx, walkableHeight, walkableClimb, y, solid);
+		rcFilterLedgeSpansImp(ctx, walkableHeight, walkableClimb, filterNeighborSlope, y, solid);	//UE
 	}
-	
+
 	ctx->stopTimer(RC_TIMER_FILTER_BORDER);
 }	
 
 
 /// @see rcHeightfield, rcConfig
-void rcFilterLedgeSpans(rcContext* ctx, const int walkableHeight, const int walkableClimb, const int yStart, const int maxYProcess,
+void rcFilterLedgeSpans(rcContext* ctx, const int walkableHeight, const int walkableClimb, const bool filterNeighborSlope, const int yStart, const int maxYProcess, //UE
 	rcHeightfield& solid)
 {
 	rcAssert(ctx);
@@ -199,7 +208,7 @@ void rcFilterLedgeSpans(rcContext* ctx, const int walkableHeight, const int walk
 
 	for (int y = yStart; y < h; ++y)
 	{
-		rcFilterLedgeSpansImp(ctx, walkableHeight, walkableClimb, y, solid);
+		rcFilterLedgeSpansImp(ctx, walkableHeight, walkableClimb, filterNeighborSlope, y, solid);
 	}
 
 	ctx->stopTimer(RC_TIMER_FILTER_BORDER);
diff --git a/Engine/Source/Runtime/Navmesh/Public/Recast/Recast.h b/Engine/Source/Runtime/Navmesh/Public/Recast/Recast.h
--- a/Engine/Source/Runtime/Navmesh/Public/Recast/Recast.h
+++ b/Engine/Source/Runtime/Navmesh/Public/Recast/Recast.h
@@ -1077,11 +1077,12 @@ NAVMESH_API void rcFilterLowHangingWalkableObstacles(rcContext* ctx, const int w
 ///  								be considered walkable. [Limit: >= 3] [Units: vx]
 ///  @param[in]		walkableClimb	Maximum ledge height that is considered to still be traversable. 
 ///  								[Limit: >=0] [Units: vx]
+///  @param[in]		filterNeighborSlope If set, a span having neighbors not within the walkable climb range will be marked as non walkable. //UE
 ///  @param[in,out]	solid			A fully built heightfield.  (All spans have been added.)
 ///  @param[in]		yStart			y coord to start at
 ///  @param[in]		maxYProcess	    Max y coords to process (yStart + maxYProcess can be more than solid.height and will be capped to solid.height)
 NAVMESH_API void rcFilterLedgeSpans(rcContext* ctx, const int walkableHeight,
-						const int walkableClimb, const int yStart, const int maxYProcess, rcHeightfield& solid);
+						const int walkableClimb, const bool filterNeighborSlope, const int yStart, const int maxYProcess, rcHeightfield& solid); //UE
 
 /// Marks spans that are ledges as not-walkable. 
 ///  @ingroup recast
@@ -1090,9 +1091,10 @@ NAVMESH_API void rcFilterLedgeSpans(rcContext* ctx, const int walkableHeight,
 ///  								be considered walkable. [Limit: >= 3] [Units: vx]
 ///  @param[in]		walkableClimb	Maximum ledge height that is considered to still be traversable. 
 ///  								[Limit: >=0] [Units: vx]
+///  @param[in]		filterNeighborSlope If set, a span having neighbors not within the walkable climb range will be marked as non walkable. //UE
 ///  @param[in,out]	solid			A fully built heightfield.  (All spans have been added.)
 NAVMESH_API void rcFilterLedgeSpans(rcContext* ctx, const int walkableHeight,
-						const int walkableClimb, rcHeightfield& solid);
+						const int walkableClimb, const bool filterNeighborSlope, rcHeightfield& solid); //UE
 
 /// Marks walkable spans as not walkable if the clearance above the span is less than the specified height. 
 ///  @ingroup recast
