Fix navmesh holes on steep slope caused by a rasterization issue (this also fixes navmesh issues in stairs).
Also added ai.nav.EnableSpanHeightRasterizationFix console command to disable the fix if needed.

diff --git a/Engine/Source/Runtime/Navmesh/Private/Recast/RecastRasterization.cpp b/Engine/Source/Runtime/Navmesh/Private/Recast/RecastRasterization.cpp
--- a/Engine/Source/Runtime/Navmesh/Private/Recast/RecastRasterization.cpp
+++ b/Engine/Source/Runtime/Navmesh/Private/Recast/RecastRasterization.cpp
@@ -24,6 +24,15 @@
 #include "Recast/Recast.h"
 #include "Recast/RecastAlloc.h"
 #include "Recast/RecastAssert.h"
+#include "HAL/ConsoleManager.h"
+
+//@UE BEGIN
+namespace UE::Recast::Private
+{
+	static bool bEnableSpanHeightRasterizationFix = true;
+	static FAutoConsoleVariableRef CVarEnableSpanHeightRasterizationFix(TEXT("ai.nav.EnableSpanHeightRasterizationFix"), bEnableSpanHeightRasterizationFix, TEXT("Active by default. Enable rasterization fix for span height."), ECVF_Default);
+}
+//@UE END
 
 inline bool overlapBounds(const rcReal* amin, const rcReal* amax, const rcReal* bmin, const rcReal* bmax)
 {
@@ -842,6 +851,7 @@ static void rasterizeTri(const rcReal* v0, const rcReal* v1, const rcReal* v2,
 				}
 			}
 		}
+		
 		for (int y = y0; y <= y1; y++)
 		{
 			int xloop0 = intMax(hf.RowExt[y + 1].MinCol, x0);
@@ -852,11 +862,10 @@ static void rasterizeTri(const rcReal* v0, const rcReal* v1, const rcReal* v2,
 				rcTempSpan& Temp = hf.tempspans[idx];
 
 				short int smin = Temp.sminmax[0];
-				short int smax = Temp.sminmax[1];
-	#if TEST_NEW_RASTERIZER
-				short int tsmin = Temp.sminmax[0];
-				short int tsmax = Temp.sminmax[1];
-	#endif
+
+				// +1 because span sample heights are computed from rcFloor instead of rcCeil (and they need to have a height of at least 1)
+				short int smax = UE::Recast::Private::bEnableSpanHeightRasterizationFix ? Temp.sminmax[1] + 1 : Temp.sminmax[1]; //UE
+
 				// reset for next triangle
 				Temp.sminmax[0] = 32000;
 				Temp.sminmax[1] = -32000;
