[Navmesh] Increase the height where the navmesh can generate.
Added RC_INCREASED_SPAN_HEIGHT_LIMIT allowing to revert to previous behavior.

diff --git a/Engine/Source/Runtime/Navmesh/Private/Recast/RecastRasterization.cpp b/Engine/Source/Runtime/Navmesh/Private/Recast/RecastRasterization.cpp
--- a/Engine/Source/Runtime/Navmesh/Private/Recast/RecastRasterization.cpp
+++ b/Engine/Source/Runtime/Navmesh/Private/Recast/RecastRasterization.cpp
@@ -528,7 +528,7 @@ static void rasterizeTri(const rcReal* v0, const rcReal* v1, const rcReal* v2,
 				bmin, bmax,
 				cs, ics, ich,
 				flagMergeThr);
-			if (outsmin != 8191)
+			if (outsmin != RC_SPAN_MAX_HEIGHT)
 			{
 				markSpanSample(hf, x, y);
 			}
@@ -903,7 +903,7 @@ static void rasterizeTri(const rcReal* v0, const rcReal* v1, const rcReal* v2,
 							bmin, bmax,
 							cs, ics, ich,
 							flagMergeThr);
-						if (outsmin != 8191)
+						if (outsmin != RC_SPAN_MAX_HEIGHT)
 						{
 							Temp.sminmax[0] = 32000;
 							Temp.sminmax[1] = -32000;
diff --git a/Engine/Source/Runtime/Navmesh/Public/Recast/Recast.h b/Engine/Source/Runtime/Navmesh/Public/Recast/Recast.h
--- a/Engine/Source/Runtime/Navmesh/Public/Recast/Recast.h
+++ b/Engine/Source/Runtime/Navmesh/Public/Recast/Recast.h
@@ -26,6 +26,12 @@
 #include "Logging/LogMacros.h"
 #include "RecastLargeWorldCoordinates.h"
 
+// @UE BEGIN
+#ifndef RC_INCREASED_SPAN_HEIGHT_LIMIT
+#define RC_INCREASED_SPAN_HEIGHT_LIMIT	1
+#endif
+// @UE END
+
 //@UE BEGIN Adding support for LWCoords.
 /// The value of PI used by Recast.
 static const rcReal RC_PI = 3.14159265358979323846;
@@ -321,7 +327,14 @@ struct rcConfig
 };
 
 /// Defines the number of bits allocated to rcSpan::smin and rcSpan::smax.
-static const int RC_SPAN_HEIGHT_BITS = 13;
+/// Using 15 bits increases the size of rcSpanData to 8 bytes but it does not impact the size of rcSpan since padding was already present.
+/// It also increases the size of rcSpanCache to 12 bytes.
+#if RC_INCREASED_SPAN_HEIGHT_LIMIT
+static constexpr int RC_SPAN_HEIGHT_BITS = 15;	// UE
+#else
+static constexpr int RC_SPAN_HEIGHT_BITS = 13;
+#endif // RC_INCREASED_SPAN_HEIGHT_LIMIT
+
 /// Defines the maximum value for rcSpan::smin and rcSpan::smax.
 static const int RC_SPAN_MAX_HEIGHT = (1<<RC_SPAN_HEIGHT_BITS)-1;
 
