[Navmesh] Increase the height where the navmesh can generate (now using 29 bits instead of 15, allowing 26000 km of altitude when using 5 cm CellHeight).
Removed unused hmin and hmax from dtTileCacheLayerHeader.
Increased navmesh version to 27.

diff --git a/Engine/Source/Runtime/NavigationSystem/Private/NavMesh/PImplRecastNavMesh.cpp b/Engine/Source/Runtime/NavigationSystem/Private/NavMesh/PImplRecastNavMesh.cpp
--- a/Engine/Source/Runtime/NavigationSystem/Private/NavMesh/PImplRecastNavMesh.cpp
+++ b/Engine/Source/Runtime/NavigationSystem/Private/NavMesh/PImplRecastNavMesh.cpp
@@ -1049,8 +1049,15 @@ void FPImplRecastNavMesh::SerializeCompressedTileCacheData(FArchive& Ar, int32 N
 		Ar << Header->bmin[i];
 		Ar << Header->bmax[i];
 	}
-	Ar << Header->hmin;
-	Ar << Header->hmax;
+	
+	if (Ar.IsLoading() && NavMeshVersion < NAVMESHVER_HEIGHT_INCREASE)
+	{
+		unsigned short hmin, hmax;
+		Ar << hmin;
+		Ar << hmax;
+		// Skipping hmin and hmax, they do not exit in new header.
+	}
+	
 	Ar << Header->width;
 	Ar << Header->height;
 	Ar << Header->minx;
diff --git a/Engine/Source/Runtime/NavigationSystem/Private/NavMesh/RecastNavMeshGenerator.cpp b/Engine/Source/Runtime/NavigationSystem/Private/NavMesh/RecastNavMeshGenerator.cpp
--- a/Engine/Source/Runtime/NavigationSystem/Private/NavMesh/RecastNavMeshGenerator.cpp
+++ b/Engine/Source/Runtime/NavigationSystem/Private/NavMesh/RecastNavMeshGenerator.cpp
@@ -3212,8 +3212,6 @@ bool FRecastTileGenerator::RecastBuildTileCache(FNavMeshBuildContext& BuildConte
 		header.maxx = (unsigned short)layer->maxx;
 		header.miny = (unsigned short)layer->miny;
 		header.maxy = (unsigned short)layer->maxy;
-		header.hmin = (unsigned short)layer->hmin;
-		header.hmax = (unsigned short)layer->hmax;
 
 		// Layer bounds in unreal coords
 		FBox LayerBBox = Recast2UnrealBox(header.bmin, header.bmax);
diff --git a/Engine/Source/Runtime/NavigationSystem/Public/NavMesh/RecastVersion.h b/Engine/Source/Runtime/NavigationSystem/Public/NavMesh/RecastVersion.h
--- a/Engine/Source/Runtime/NavigationSystem/Public/NavMesh/RecastVersion.h
+++ b/Engine/Source/Runtime/NavigationSystem/Public/NavMesh/RecastVersion.h
@@ -29,6 +29,7 @@
 #define NAVMESHVER_TILE_RESOLUTIONS_CELLHEIGHT	24 // Addition of CellHeight in the resolution params, deprecating the original CellHeight.
 #define NAVMESHVER_1_VOXEL_AGENT_STEEP_SLOPE_FILTER_FIX	25 // Fix, remove steep slope filtering during heightfield ledge filtering when the agent radius is included into a single voxel
 #define NAVMESHVER_TILE_RESOLUTIONS_AGENTMAXSTEPHEIGHT 26	// Addition of AgentMaxStepHeight in the resolution params, deprecating the original AgentMaxStepHeight.
+#define NAVMESHVER_HEIGHT_INCREASE 27						// Increased rcSpanData min/max height range (changes to dtTileCacheLayerHeader).
 
-#define NAVMESHVER_LATEST				NAVMESHVER_TILE_RESOLUTIONS_AGENTMAXSTEPHEIGHT
+#define NAVMESHVER_LATEST				NAVMESHVER_HEIGHT_INCREASE
 #define NAVMESHVER_MIN_COMPATIBLE		NAVMESHVER_LWCOORDS_OPTIMIZATION
diff --git a/Engine/Source/Runtime/Navmesh/Private/DetourTileCache/DetourTileCacheBuilder.cpp b/Engine/Source/Runtime/Navmesh/Private/DetourTileCache/DetourTileCacheBuilder.cpp
--- a/Engine/Source/Runtime/Navmesh/Private/DetourTileCache/DetourTileCacheBuilder.cpp
+++ b/Engine/Source/Runtime/Navmesh/Private/DetourTileCache/DetourTileCacheBuilder.cpp
@@ -2762,8 +2762,6 @@ bool dtTileCacheHeaderSwapEndian(unsigned char* data, const int dataSize)
 	dtSwapEndian(&header->bmax[0]);
 	dtSwapEndian(&header->bmax[1]);
 	dtSwapEndian(&header->bmax[2]);
-	dtSwapEndian(&header->hmin);
-	dtSwapEndian(&header->hmax);
 
 	dtSwapEndian(&header->width);
 	dtSwapEndian(&header->height);
diff --git a/Engine/Source/Runtime/Navmesh/Private/Recast/Recast.cpp b/Engine/Source/Runtime/Navmesh/Private/Recast/Recast.cpp
--- a/Engine/Source/Runtime/Navmesh/Private/Recast/Recast.cpp
+++ b/Engine/Source/Runtime/Navmesh/Private/Recast/Recast.cpp
@@ -279,8 +279,9 @@ bool rcCreateHeightfield(rcContext* /*ctx*/, rcHeightfield& hf, int width, int h
 	{
 		for (int j = 0; j < hf.width + 2; j++)
 		{
-			hf.tempspans[i * (hf.width + 2) + j].sminmax[0] = 32000;
-			hf.tempspans[i * (hf.width + 2) + j].sminmax[1] = -32000;
+			constexpr int RANGE = INT_MAX;
+			hf.tempspans[i * (hf.width + 2) + j].sminmax[0] = RANGE;
+			hf.tempspans[i * (hf.width + 2) + j].sminmax[1] = -RANGE;
 		}
 	}
 
@@ -486,8 +487,6 @@ bool rcBuildCompactHeightfield(rcContext* ctx, const int walkableHeight, const i
 	}
 	memset(chf.areas, RC_NULL_AREA, sizeof(unsigned char)*spanCount);
 	
-	const int MAX_HEIGHT = 0xffff;
-	
 	// Fill in cells and spans.
 	int idx = 0;
 	for (int y = 0; y < h; ++y)
@@ -504,9 +503,9 @@ bool rcBuildCompactHeightfield(rcContext* ctx, const int walkableHeight, const i
 			{
 				if (s->data.area != RC_NULL_AREA)
 				{
-					const int bot = (int)s->data.smax;
-					const int top = s->next ? (int)s->next->data.smin : MAX_HEIGHT;
-					chf.spans[idx].y = (unsigned short)rcClamp(bot, 0, 0xffff);
+					const rcSpanInt bot = (int)s->data.smax;
+					const rcSpanInt top = s->next ? (int)s->next->data.smin : RC_SPAN_MAX_HEIGHT;
+					chf.spans[idx].y = rcClamp(bot, 0, UINT_MAX);
 					chf.spans[idx].h = (unsigned char)rcClamp(top - bot, 0, 0xff);
 					chf.areas[idx] = s->data.area;
 					idx++;
diff --git a/Engine/Source/Runtime/Navmesh/Private/Recast/RecastArea.cpp b/Engine/Source/Runtime/Navmesh/Private/Recast/RecastArea.cpp
--- a/Engine/Source/Runtime/Navmesh/Private/Recast/RecastArea.cpp
+++ b/Engine/Source/Runtime/Navmesh/Private/Recast/RecastArea.cpp
@@ -274,7 +274,7 @@ void SeedWalkableAreasErodeFiltered(int w, int h, int lowSpanHeight, int lowSpan
 			const rcCompactCell& c = chf.cells[x + y*w];
 
 			// initialize to max height in case it needs to start with a low span
-			int32 NextAllowedLowBase = 0xffff;
+			unsigned int NextAllowedLowBase = UINT_MAX;
 
 			for (int ni = (int)c.index, i = (int)(c.index + c.count) - 1; i >= ni; i--)
 			{
@@ -338,7 +338,7 @@ void FilterLowHeightSpans(int w, int h, int lowSpanHeight, unsigned char areaId,
 			const rcCompactCell& c = chf.cells[x + y*w];
 
 			// initialize to max height in case it needs to start with a low span
-			int32 NextAllowedLowBase = 0xffff;
+			unsigned int NextAllowedLowBase = UINT_MAX;
 
 			for (int ni = (int)c.index, i = (int)(c.index + c.count) - 1; i >= ni; i--)
 			{
diff --git a/Engine/Source/Runtime/Navmesh/Private/Recast/RecastFilter.cpp b/Engine/Source/Runtime/Navmesh/Private/Recast/RecastFilter.cpp
--- a/Engine/Source/Runtime/Navmesh/Private/Recast/RecastFilter.cpp
+++ b/Engine/Source/Runtime/Navmesh/Private/Recast/RecastFilter.cpp
@@ -257,7 +257,6 @@ void rcFilterWalkableLowHeightSpansSequences(rcContext* ctx, int walkableHeight,
 
 	const int w = solid.width;
 	const int h = solid.height;
-	const int MAX_HEIGHT = 0xffff;
 
 	const int32 MaxSpans = 64;
 	rcCompactSpan SpanList[MaxSpans];
@@ -276,9 +275,9 @@ void rcFilterWalkableLowHeightSpansSequences(rcContext* ctx, int walkableHeight,
 			NumSpans = 0;
 			for (rcSpan* s = solid.spans[x + y*w]; s; s = s->next)
 			{
-				const int bot = (int)s->data.smax;
-				const int top = s->next ? (int)s->next->data.smin : MAX_HEIGHT;
-				SpanList[NumSpans].y = (unsigned short)rcClamp(bot, 0, 0xffff);
+				const rcSpanInt bot = s->data.smax;
+				const rcSpanInt top = s->next ? s->next->data.smin : RC_SPAN_MAX_HEIGHT;
+				SpanList[NumSpans].y = rcClamp(bot, 0, UINT_MAX);
 				SpanList[NumSpans].h = (unsigned char)rcClamp(top - bot, 0, 0xff);
 				SpanList[NumSpans].reg = s->data.area;
 				
@@ -289,7 +288,7 @@ void rcFilterWalkableLowHeightSpansSequences(rcContext* ctx, int walkableHeight,
 				}
 			}
 
-			int32 NextAllowedBase = 0xffff;
+			unsigned int NextAllowedBase = UINT_MAX;
 			for (int32 Idx = NumSpans - 1; Idx >= 0; Idx--)
 			{
 				if (SpanList[Idx].h < walkableHeight)
diff --git a/Engine/Source/Runtime/Navmesh/Private/Recast/RecastLayers.cpp b/Engine/Source/Runtime/Navmesh/Private/Recast/RecastLayers.cpp
--- a/Engine/Source/Runtime/Navmesh/Private/Recast/RecastLayers.cpp
+++ b/Engine/Source/Runtime/Navmesh/Private/Recast/RecastLayers.cpp
@@ -31,7 +31,7 @@ struct rcLayerRegionMonotone
 	int chunkId;
 	rcIntArray neis;
 	rcIntArray layers;
-	unsigned short ymin, ymax;
+	unsigned int ymin, ymax;
 	unsigned short layerId;		// Layer ID
 	unsigned char base : 1;		// Flag indicating if the region is the base of merged regions.
 	unsigned char remap : 1;
@@ -248,7 +248,7 @@ static bool CollectLayerRegionsMonotone(rcContext* ctx, rcCompactHeightfield& ch
 	for (int i = 0; i < nregs; ++i)
 	{
 		regs[i].layerId = 0xffff;
-		regs[i].ymin = 0xffff;
+		regs[i].ymin = UINT_MAX;
 		regs[i].ymax = 0;
 	}
 
@@ -464,7 +464,7 @@ static bool CollectLayerRegionsChunky(rcContext* ctx, rcCompactHeightfield& chf,
 	for (int i = 0; i < nregs; ++i)
 	{
 		regs[i].layerId = 0xffff;
-		regs[i].ymin = 0xffff;
+		regs[i].ymin = UINT_MAX;
 		regs[i].ymax = 0;
 	}
 
@@ -571,8 +571,8 @@ static bool SplitAndStoreLayerRegions(rcContext* ctx, rcCompactHeightfield& chf,
 				if (root.layers.contains(nei))
 					continue;
 				// Skip if the height range would become too large.
-				const int ymin = rcMin(root.ymin, regn.ymin);
-				const int ymax = rcMax(root.ymax, regn.ymax);
+				const unsigned int ymin = rcMin(root.ymin, regn.ymin);
+				const unsigned int ymax = rcMax(root.ymax, regn.ymax);
 				if ((ymax - ymin) >= 255)
 					continue;
 
@@ -616,8 +616,8 @@ static bool SplitAndStoreLayerRegions(rcContext* ctx, rcCompactHeightfield& chf,
 				if (!overlapRange(ri.ymin,ri.ymax+mergeHeight, rj.ymin,rj.ymax+mergeHeight))
 					continue;
 				// Skip if the height range would become too large.
-				const int ymin = rcMin(ri.ymin, rj.ymin);
-				const int ymax = rcMax(ri.ymax, rj.ymax);
+				const unsigned int ymin = rcMin(ri.ymin, rj.ymin);
+				const unsigned int ymax = rcMax(ri.ymax, rj.ymax);
 				if ((ymax - ymin) >= 255)
 					continue;
 
@@ -780,13 +780,13 @@ static bool SplitAndStoreLayerRegions(rcContext* ctx, rcCompactHeightfield& chf,
 		memset(layer->cons, 0, gridSize);
 
 		// Find layer height bounds.
-		int hmin = 0, hmax = 0;
+		unsigned int hmin = 0, hmax = 0;
 		for (int j = 0; j < nregs; ++j)
 		{
 			if (regs[j].base && regs[j].layerId == curId)
 			{
-				hmin = (int)regs[j].ymin;
-				hmax = (int)regs[j].ymax;
+				hmin = regs[j].ymin;
+				hmax = regs[j].ymax;
 			}
 		}
 
@@ -987,7 +987,7 @@ struct rcLayerRegion
 	rcIntArray layers;
 	rcIntArray connections;
 	unsigned short layerId;
-	unsigned short ymin, ymax;
+	unsigned int ymin, ymax;
 	unsigned char remap : 1;
 	unsigned char visited : 1;
 	unsigned char base : 1;
@@ -1436,13 +1436,13 @@ bool rcBuildHeightfieldLayers(rcContext* ctx, rcCompactHeightfield& chf,
 		memset(layer->cons, 0, gridSize);
 
 		// Find layer height bounds.
-		int hmin = 0, hmax = 0;
+		unsigned int hmin = 0, hmax = 0;
 		for (int j = 0; j < nreg; ++j)
 		{
 			if (regions[j].base && regions[j].layerId == curId)
 			{
-				hmin = (int)regions[j].ymin;
-				hmax = (int)regions[j].ymax;
+				hmin = regions[j].ymin;
+				hmax = regions[j].ymax;
 			}
 		}
 
diff --git a/Engine/Source/Runtime/Navmesh/Private/Recast/RecastRasterization.cpp b/Engine/Source/Runtime/Navmesh/Private/Recast/RecastRasterization.cpp
--- a/Engine/Source/Runtime/Navmesh/Private/Recast/RecastRasterization.cpp
+++ b/Engine/Source/Runtime/Navmesh/Private/Recast/RecastRasterization.cpp
@@ -94,7 +94,7 @@ static void freeSpan(rcHeightfield& hf, rcSpan* ptr)
 }
 
 static void addSpan(rcHeightfield& hf, const int x, const int y,
-					const unsigned short smin, const unsigned short smax,
+					const rcSpanInt smin, const rcSpanInt smax,
 					const unsigned char area, const int flagMergeThr)
 {
 	
@@ -186,7 +186,7 @@ static void addSpan(rcHeightfield& hf, const int x, const int y,
 ///
 /// @see rcHeightfield, rcSpan.
 void rcAddSpan(rcContext* /*ctx*/, rcHeightfield& hf, const int x, const int y,
-			   const unsigned short smin, const unsigned short smax,
+			   const rcSpanInt smin, const rcSpanInt smax,
 			   const unsigned char area, const int flagMergeThr)
 {
 //	rcAssert(ctx);
@@ -322,7 +322,7 @@ static inline int SampleIndex(rcHeightfield const& hf, const int x, const int y)
 	return x + 1 + (y + 1)*(hf.width + 2);
 }
 
-static inline void addSpanSample(rcHeightfield& hf, const int x, const int y, short int sint)
+static inline void addSpanSample(rcHeightfield& hf, const int x, const int y, int sint)
 {
 	addFlatSpanSample(hf, x, y);
 	int idx = SampleIndex(hf, x, y);
@@ -355,7 +355,7 @@ static inline void intersectZ(const rcReal* v0, const rcReal* edge, rcReal cz, r
 
 static void rasterizeTriTest(const rcReal* v0, const rcReal* v1, const rcReal* v2,
 						int xtest, int ytest,
-						short int& outsmin, short int& outsmax, 
+						int& outsmin, int& outsmax, 
 						 const unsigned char area, rcHeightfield& hf,
 						 const rcReal* bmin, const rcReal* bmax,
 						 const rcReal cs, const rcReal ics, const rcReal ich,
@@ -435,8 +435,8 @@ static void rasterizeTriTest(const rcReal* v0, const rcReal* v1, const rcReal* v
 			if (smax > by) smax = by;
 
 			// Snap the span to the heightfield height grid.
-			unsigned short ismin = (unsigned short)rcClamp((int)rcFloor(smin * ich), 0, RC_SPAN_MAX_HEIGHT);
-			unsigned short ismax = (unsigned short)rcClamp((int)rcCeil(smax * ich), (int)ismin+1, RC_SPAN_MAX_HEIGHT);
+			rcSpanInt ismin = (rcSpanInt)rcClamp((int)rcFloor(smin * ich), 0, RC_SPAN_MAX_HEIGHT);
+			rcSpanInt ismax = (rcSpanInt)rcClamp((int)rcCeil(smax * ich), (int)ismin+1, RC_SPAN_MAX_HEIGHT);
 			outsmin = ismin;
 			outsmax = ismax;
 		}
@@ -494,8 +494,8 @@ static void rasterizeTri(const rcReal* v0, const rcReal* v1, const rcReal* v2,
 		if (triangle_smax > by) triangle_smax = by;
 
 		// Snap the span to the heightfield height grid.
-		unsigned short triangle_ismin = (unsigned short)rcClamp((int)rcFloor(triangle_smin * ich), 0, RC_SPAN_MAX_HEIGHT);
-		unsigned short triangle_ismax = (unsigned short)rcClamp((int)rcCeil(triangle_smax * ich), (int)triangle_ismin+1, RC_SPAN_MAX_HEIGHT);
+		rcSpanInt triangle_ismin = (rcSpanInt)rcClamp((int)rcFloor(triangle_smin * ich), 0, RC_SPAN_MAX_HEIGHT);
+		rcSpanInt triangle_ismax = (rcSpanInt)rcClamp((int)rcCeil(triangle_smax * ich), (int)triangle_ismin+1, RC_SPAN_MAX_HEIGHT);
 		const int projectSpanToBottom = rasterizationMasks != nullptr ? (projectTriToBottom & rasterizationMasks[x0+y0*w]) : projectTriToBottom;	//UE
 		if (projectSpanToBottom) //UE
 		{
@@ -506,8 +506,10 @@ static void rasterizeTri(const rcReal* v0, const rcReal* v1, const rcReal* v2,
 		return;
 	}
 
-	const short int triangle_ismin = (short int)rcClamp((int)rcFloor(triangle_smin * ich), -32000, 32000);
-	const short int triangle_ismax = (short int)rcClamp((int)rcFloor(triangle_smax * ich), -32000, 32000);
+	constexpr int RANGE = INT_MAX;
+	
+	const int triangle_ismin = rcClamp((int)rcFloor(triangle_smin * ich), -RANGE, RANGE);
+	const int triangle_ismax = rcClamp((int)rcFloor(triangle_smax * ich), -RANGE, RANGE);
 
 	x0 = intMax(x0, 0);
 	int x1_edge = intMin(x1, w);
@@ -521,7 +523,7 @@ static void rasterizeTri(const rcReal* v0, const rcReal* v1, const rcReal* v2,
 	{
 		for (int x = x0; x <= x1; x++)
 		{
-			short int outsmin, outsmax;
+			int outsmin, outsmax;
 			rasterizeTriTest(v0, v1, v2, x, y, 
 				outsmin, outsmax, 
 				area, hf,
@@ -661,8 +663,8 @@ static void rasterizeTri(const rcReal* v0, const rcReal* v1, const rcReal* v2,
 		if (rasterizationMasks == nullptr) //UE
 		{
 			// Snap the span to the heightfield height grid.
-			unsigned short triangle_ismin_clamp = (unsigned short)rcClamp((int)triangle_ismin, 0, RC_SPAN_MAX_HEIGHT);
-			const unsigned short triangle_ismax_clamp = (unsigned short)rcClamp((int)triangle_ismax, (int)triangle_ismin_clamp+1, RC_SPAN_MAX_HEIGHT);
+			rcSpanInt triangle_ismin_clamp = (rcSpanInt)rcClamp((int)triangle_ismin, 0, RC_SPAN_MAX_HEIGHT);
+			const rcSpanInt triangle_ismax_clamp = (rcSpanInt)rcClamp((int)triangle_ismax, (int)triangle_ismin_clamp+1, RC_SPAN_MAX_HEIGHT);
 			if (projectTriToBottom) //UE
 			{
 				triangle_ismin_clamp = 0; //UE
@@ -692,8 +694,8 @@ static void rasterizeTri(const rcReal* v0, const rcReal* v1, const rcReal* v2,
 				for (int x = xloop0; x <= xloop1; x++)
 				{
 					// Snap the span to the heightfield height grid.
-					unsigned short triangle_ismin_clamp = (unsigned short)rcClamp((int)triangle_ismin, 0, RC_SPAN_MAX_HEIGHT);
-					const unsigned short triangle_ismax_clamp = (unsigned short)rcClamp((int)triangle_ismax, (int)triangle_ismin_clamp+1, RC_SPAN_MAX_HEIGHT);
+					rcSpanInt triangle_ismin_clamp = (rcSpanInt)rcClamp((int)triangle_ismin, 0, RC_SPAN_MAX_HEIGHT);
+					const rcSpanInt triangle_ismax_clamp = (rcSpanInt)rcClamp((int)triangle_ismax, (int)triangle_ismin_clamp+1, RC_SPAN_MAX_HEIGHT);
 					const int projectSpanToBottom = projectTriToBottom & rasterizationMasks[x+y*w];		//UE
 					if (projectSpanToBottom) //UE
 					{
@@ -727,7 +729,7 @@ static void rasterizeTri(const rcReal* v0, const rcReal* v1, const rcReal* v2,
 			if (intverts[basevert][0] >= x0 && intverts[basevert][0] <= x1 && intverts[basevert][1] >= y0 && intverts[basevert][1] <= y1)
 			{
 				rcReal sfloat = vertarray[basevert][1] - bmin[1];
-				short int sint = (short int)rcClamp((int)rcFloor(sfloat * ich), -32000, 32000);
+				int sint = (int)rcClamp((int)rcFloor(sfloat * ich), -RANGE, RANGE);
 	#if TEST_NEW_RASTERIZER
 				rcAssert(sint >= triangle_ismin - 1 && sint <= triangle_ismax + 1);
 	#endif
@@ -765,7 +767,7 @@ static void rasterizeTri(const rcReal* v0, const rcReal* v1, const rcReal* v2,
 					if (y >= y0 && y <= y1)
 					{
 						rcReal sfloat = temppnt[1] - bmin[1];
-						short int sint = (short int)rcClamp((int)rcFloor(sfloat * ich), -32000, 32000);
+						short int sint = (short int)rcClamp((int)rcFloor(sfloat * ich), -RANGE, RANGE);
 #if TEST_NEW_RASTERIZER
 						rcAssert(sint >= triangle_ismin - 1 && sint <= triangle_ismax + 1);
 #endif
@@ -806,7 +808,7 @@ static void rasterizeTri(const rcReal* v0, const rcReal* v1, const rcReal* v2,
 						if (x >= x0 && x <= x1)
 						{
 							rcReal sfloat = Inter[i][1] - bmin[1];
-							short int sint = (short int)rcClamp((int)rcFloor(sfloat * ich), -32000, 32000);
+							short int sint = (short int)rcClamp((int)rcFloor(sfloat * ich), -RANGE, RANGE);
 #if TEST_NEW_RASTERIZER
 							rcAssert(sint >= triangle_ismin - 1 && sint <= triangle_ismax + 1);
 #endif
@@ -835,7 +837,7 @@ static void rasterizeTri(const rcReal* v0, const rcReal* v1, const rcReal* v2,
 						}
 						for (int x = xloop0; x <= xloop1; x++, sfloat += ds)
 						{
-							short int sint = (short int)rcClamp((int)rcFloor(sfloat * ich), -32000, 32000);
+							short int sint = (short int)rcClamp((int)rcFloor(sfloat * ich), -RANGE, RANGE);
 #if TEST_NEW_RASTERIZER
 							rcAssert(sint >= triangle_ismin - 1 && sint <= triangle_ismax + 1);
 #endif
@@ -861,20 +863,20 @@ static void rasterizeTri(const rcReal* v0, const rcReal* v1, const rcReal* v2,
 				int idx = SampleIndex(hf, x, y);
 				rcTempSpan& Temp = hf.tempspans[idx];
 
-				short int smin = Temp.sminmax[0];
+				int smin = Temp.sminmax[0];
 
 				// +1 because span sample heights are computed from rcFloor instead of rcCeil (and they need to have a height of at least 1)
-				short int smax = UE::Recast::Private::bEnableSpanHeightRasterizationFix ? Temp.sminmax[1] + 1 : Temp.sminmax[1]; //UE
+				int smax = UE::Recast::Private::bEnableSpanHeightRasterizationFix ? Temp.sminmax[1] + 1 : Temp.sminmax[1]; //UE
 
 				// reset for next triangle
-				Temp.sminmax[0] = 32000;
-				Temp.sminmax[1] = -32000;
+				Temp.sminmax[0] = RANGE;
+				Temp.sminmax[1] = -RANGE;
 
 				// Skip the span if it is outside the heightfield bbox
 				if (smin >= RC_SPAN_MAX_HEIGHT || smax < 0) continue;
 
-				smin = (short)intMax(smin, 0);
-				smax = (short)intMin(intMax(smax,smin+1), RC_SPAN_MAX_HEIGHT);
+				smin = intMax(smin, 0);
+				smax = intMin(intMax(smax,smin+1), RC_SPAN_MAX_HEIGHT);
 				const int projectSpanToBottom = rasterizationMasks != nullptr ? (projectTriToBottom & rasterizationMasks[x+y*w]) : projectTriToBottom; //UE
 				if (projectSpanToBottom) //UE
 				{
@@ -883,7 +885,7 @@ static void rasterizeTri(const rcReal* v0, const rcReal* v1, const rcReal* v2,
 
 	#if TEST_NEW_RASTERIZER
 				{
-					short int outsmin, outsmax;
+					int outsmin, outsmax;
 					rasterizeTriTest(v0, v1, v2, x, y, 
 						outsmin, outsmax, 
 						area, hf,
@@ -895,8 +897,8 @@ static void rasterizeTri(const rcReal* v0, const rcReal* v1, const rcReal* v2,
 						outsmax > smax + tol || outsmax < smax - tol
 						)
 					{
-						Temp.sminmax[0] = 32000;
-						Temp.sminmax[1] = -32000;
+						Temp.sminmax[0] = RANGE;
+						Temp.sminmax[1] = -RANGE;
 						rasterizeTriTest(v0, v1, v2, x, y, 
 							outsmin, outsmax, 
 							area, hf,
@@ -905,8 +907,8 @@ static void rasterizeTri(const rcReal* v0, const rcReal* v1, const rcReal* v2,
 							flagMergeThr);
 						if (outsmin != RC_SPAN_MAX_HEIGHT)
 						{
-							Temp.sminmax[0] = 32000;
-							Temp.sminmax[1] = -32000;
+							Temp.sminmax[0] = RANGE;
+							Temp.sminmax[1] = -RANGE;
 						}
 					}
 
@@ -930,7 +932,7 @@ static void rasterizeTri(const rcReal* v0, const rcReal* v1, const rcReal* v2,
 		{
 			int idx = SampleIndex(hf, x, y);
 			rcTempSpan& Temp = hf.tempspans[idx];
-			rcAssert(Temp.sminmax[0] == 32000 && Temp.sminmax[1] == -32000);
+			rcAssert(Temp.sminmax[0] == RANGE && Temp.sminmax[1] == -RANGE);
 		}
 	}
 	rcEdgeHit& Hits = hfEdgeHits[h];
diff --git a/Engine/Source/Runtime/Navmesh/Public/DetourTileCache/DetourTileCacheBuilder.h b/Engine/Source/Runtime/Navmesh/Public/DetourTileCache/DetourTileCacheBuilder.h
--- a/Engine/Source/Runtime/Navmesh/Public/DetourTileCache/DetourTileCacheBuilder.h
+++ b/Engine/Source/Runtime/Navmesh/Public/DetourTileCache/DetourTileCacheBuilder.h
@@ -29,7 +29,7 @@
 #include "HAL/PlatformString.h"
 
 static const int DT_TILECACHE_MAGIC = 'D'<<24 | 'T'<<16 | 'L'<<8 | 'R'; ///< 'DTLR';
-static const int DT_TILECACHE_VERSION = 1;
+static const int DT_TILECACHE_VERSION = 2;
 
 static const unsigned char DT_TILECACHE_NULL_AREA = 0;
 static const unsigned char DT_TILECACHE_WALKABLE_AREA = 63;
@@ -39,7 +39,6 @@ struct dtTileCacheLayerHeader
 {
 	short version;							///< Data version
 	int tx,ty,tlayer;
-	unsigned short hmin, hmax;				///< Height min/max range
 	unsigned short width, height;			///< Dimension of the layer.
 	unsigned short minx, maxx, miny, maxy;	///< Usable sub-region.
 	// These should be at the bottom, as they are less often used than the rest of the data.
diff --git a/Engine/Source/Runtime/Navmesh/Public/Recast/Recast.h b/Engine/Source/Runtime/Navmesh/Public/Recast/Recast.h
--- a/Engine/Source/Runtime/Navmesh/Public/Recast/Recast.h
+++ b/Engine/Source/Runtime/Navmesh/Public/Recast/Recast.h
@@ -26,12 +26,6 @@
 #include "Logging/LogMacros.h"
 #include "RecastLargeWorldCoordinates.h"
 
-// @UE BEGIN
-#ifndef RC_INCREASED_SPAN_HEIGHT_LIMIT
-#define RC_INCREASED_SPAN_HEIGHT_LIMIT	1
-#endif
-// @UE END
-
 //@UE BEGIN Adding support for LWCoords.
 /// The value of PI used by Recast.
 static const rcReal RC_PI = 3.14159265358979323846;
@@ -326,29 +320,28 @@ struct rcConfig
 	rcReal detailSampleMaxError;
 };
 
-/// Defines the number of bits allocated to rcSpan::smin and rcSpan::smax.
-/// Using 15 bits increases the size of rcSpanData to 8 bytes but it does not impact the size of rcSpan since padding was already present.
+/// Defines the number of bits allocated to rcSpanData::smin and rcSpanData::smax.
+/// Using 29 bits increases the size of rcSpanData to 8 bytes but it does not impact the size of rcSpan since padding was already present.
 /// It also increases the size of rcSpanCache to 12 bytes.
-#if RC_INCREASED_SPAN_HEIGHT_LIMIT
-static constexpr int RC_SPAN_HEIGHT_BITS = 15;	// UE
-#else
-static constexpr int RC_SPAN_HEIGHT_BITS = 13;
-#endif // RC_INCREASED_SPAN_HEIGHT_LIMIT
+/// Size of rcTempSpan also increases to 8 bytes.
+static constexpr int RC_SPAN_HEIGHT_BITS = 29;	// UE
 
-/// Defines the maximum value for rcSpan::smin and rcSpan::smax.
+/// Defines the maximum value for rcSpanData::smin and rcSpanData::smax.
 static const int RC_SPAN_MAX_HEIGHT = (1<<RC_SPAN_HEIGHT_BITS)-1;
 
 /// The number of spans allocated per span spool.
 /// @see rcSpanPool
 static const int RC_SPANS_PER_POOL = 2048;
 
+typedef unsigned int rcSpanInt;
+
 /// Represents data of span in a heightfield.
 /// @see rcHeightfield
 struct rcSpanData
 {
-	unsigned int smin : RC_SPAN_HEIGHT_BITS;	///< The lower limit of the span. [Limit: < #smax]
-	unsigned int smax : RC_SPAN_HEIGHT_BITS;	///< The upper limit of the span. [Limit: <= #RC_SPAN_MAX_HEIGHT]
-	unsigned int area : 6;			///< The area id assigned to the span.
+	rcSpanInt smin : RC_SPAN_HEIGHT_BITS;	///< The lower limit of the span. [Limit: < #smax]
+	rcSpanInt smax : RC_SPAN_HEIGHT_BITS;	///< The upper limit of the span. [Limit: <= #RC_SPAN_MAX_HEIGHT]
+	unsigned int area : 6;					///< The area id assigned to the span.
 };
 
 struct rcSpanCache
@@ -389,7 +382,7 @@ struct rcEdgeHit
 };
 struct rcTempSpan
 {
-	short int sminmax[2];			///< The lower and upper limit of the span. [Limit: < #smax]
+	int sminmax[2];			///< The lower and upper limit of the span. [Limit: < #smax]
 };
 #endif
 
@@ -424,9 +417,9 @@ struct rcCompactCell
 /// Represents a span of unobstructed space within a compact heightfield.
 struct rcCompactSpan
 {
-	unsigned short y;			///< The lower extent of the span. (Measured from the heightfield's base.)
-	unsigned short reg;			///< The id of the region the span belongs to. (Or zero if not in a region.)
+	unsigned int y;				///< The lower extent of the span. (Measured from the heightfield's base.)
 	unsigned int con;			///< Packed neighbor connection data.
+	unsigned short reg;			///< The id of the region the span belongs to. (Or zero if not in a region.)
 	unsigned char h;			///< The height of the span.  (Measured from #y.)
 };
 
