[Engine]

diff --git a/Engine/Source/Runtime/Navmesh/Public/Detour/DetourNavMesh.h b/Engine/Source/Runtime/Navmesh/Public/Detour/DetourNavMesh.h
--- a/Engine/Source/Runtime/Navmesh/Public/Detour/DetourNavMesh.h
+++ b/Engine/Source/Runtime/Navmesh/Public/Detour/DetourNavMesh.h
@@ -493,11 +493,11 @@ struct dtNavMeshParams
 
 /// A navigation mesh based on tiles of convex polygons.
 /// @ingroup detour
-class NAVMESH_API dtNavMesh
+class dtNavMesh
 {
 public:
-	dtNavMesh();
-	~dtNavMesh();
+	NAVMESH_API dtNavMesh();
+	NAVMESH_API ~dtNavMesh();
 
 	/// @{
 	/// @name Initialization and Tile Management
@@ -505,7 +505,7 @@ public:
 	/// Initializes the navigation mesh for tiled use.
 	///  @param[in]	params		Initialization parameters.
 	/// @return The status flags for the operation.
-	dtStatus init(const dtNavMeshParams* params);
+	NAVMESH_API dtStatus init(const dtNavMeshParams* params);
 
 	/// Initializes the navigation mesh for single tile use.
 	///  @param[in]	data		Data of the new tile. (See: #dtCreateNavMeshData)
@@ -513,10 +513,10 @@ public:
 	///  @param[in]	flags		The tile flags. (See: #dtTileFlags)
 	/// @return The status flags for the operation.
 	///  @see dtCreateNavMeshData
-	dtStatus init(unsigned char* data, const int dataSize, const int flags);
+	NAVMESH_API dtStatus init(unsigned char* data, const int dataSize, const int flags);
 	
 	/// The navigation mesh initialization params.
-	const dtNavMeshParams* getParams() const;
+	NAVMESH_API const dtNavMeshParams* getParams() const;
 
 	/// Adds a tile to the navigation mesh.
 	///  @param[in]		data		Data for the new tile mesh. (See: #dtCreateNavMeshData)
@@ -525,14 +525,14 @@ public:
 	///  @param[in]		lastRef		The desired reference for the tile. (When reloading a tile.) [opt] [Default: 0]
 	///  @param[out]	result		The tile reference. (If the tile was succesfully added.) [opt]
 	/// @return The status flags for the operation.
-	dtStatus addTile(unsigned char* data, int dataSize, int flags, dtTileRef lastRef, dtTileRef* result);
+	NAVMESH_API dtStatus addTile(unsigned char* data, int dataSize, int flags, dtTileRef lastRef, dtTileRef* result);
 	
 	/// Removes the specified tile from the navigation mesh.
 	///  @param[in]		ref			The reference of the tile to remove.
 	///  @param[out]	data		Data associated with deleted tile.
 	///  @param[out]	dataSize	Size of the data associated with deleted tile.
 	/// @return The status flags for the operation.
-	dtStatus removeTile(dtTileRef ref, unsigned char** data, int* dataSize);
+	NAVMESH_API dtStatus removeTile(dtTileRef ref, unsigned char** data, int* dataSize);
 
 	/// @}
 
@@ -543,27 +543,27 @@ public:
 	///  @param[in]	pos  The world position for the query. [(x, y, z)]
 	///  @param[out]	tx		The tile's x-location. (x, y)
 	///  @param[out]	ty		The tile's y-location. (x, y)
-	void calcTileLoc(const dtReal* pos, int* tx, int* ty) const;
+	NAVMESH_API void calcTileLoc(const dtReal* pos, int* tx, int* ty) const;
 
 
 	/// Calculates whether the tile grid location for the specified world position
 	/// can fit in the tile indices type (currently an int)
 	///  @param[in]	pos			The world position for the query. [(x, y, z)]
-	bool isTileLocInValidRange(const dtReal* pos) const;
+	NAVMESH_API bool isTileLocInValidRange(const dtReal* pos) const;
 
 	/// Gets the tile at the specified grid location.
 	///  @param[in]	x		The tile's x-location. (x, y, layer)
 	///  @param[in]	y		The tile's y-location. (x, y, layer)
 	///  @param[in]	layer	The tile's layer. (x, y, layer)
 	/// @return The tile, or null if the tile does not exist.
-	const dtMeshTile* getTileAt(const int x, const int y, const int layer) const;
+	NAVMESH_API const dtMeshTile* getTileAt(const int x, const int y, const int layer) const;
 
 // @UE BEGIN
 	/// Gets number of tiles at the specified grid location. (All layers.)
 	///  @param[in]		x			The tile's x-location. (x, y)
 	///  @param[in]		y			The tile's y-location. (x, y)
 	/// @return The number of tiles in grid.
-	int getTileCountAt(const int x, const int y) const;
+	NAVMESH_API int getTileCountAt(const int x, const int y) const;
 // @UE END
 
 	/// Gets all tiles at the specified grid location. (All layers.)
@@ -572,7 +572,7 @@ public:
 	///  @param[out]	tiles		A pointer to an array of tiles that will hold the result.
 	///  @param[in]		maxTiles	The maximum tiles the tiles parameter can hold.
 	/// @return The number of tiles returned in the tiles array.
-	int getTilesAt(const int x, const int y,
+	NAVMESH_API int getTilesAt(const int x, const int y,
 				   dtMeshTile const** tiles, const int maxTiles) const;
 	
 	/// Gets the tile reference for the tile at specified grid location.
@@ -580,55 +580,55 @@ public:
 	///  @param[in]	y		The tile's y-location. (x, y, layer)
 	///  @param[in]	layer	The tile's layer. (x, y, layer)
 	/// @return The tile reference of the tile, or 0 if there is none.
-	dtTileRef getTileRefAt(int x, int y, int layer) const;
+	NAVMESH_API dtTileRef getTileRefAt(int x, int y, int layer) const;
 
 	/// Gets the tile reference for the specified tile.
 	///  @param[in]	tile	The tile.
 	/// @return The tile reference of the tile.
-	dtTileRef getTileRef(const dtMeshTile* tile) const;
+	NAVMESH_API dtTileRef getTileRef(const dtMeshTile* tile) const;
 
 	/// Gets the tile for the specified tile reference.
 	///  @param[in]	ref		The tile reference of the tile to retrieve.
 	/// @return The tile for the specified reference, or null if the 
 	///		reference is invalid.
-	const dtMeshTile* getTileByRef(dtTileRef ref) const;
+	NAVMESH_API const dtMeshTile* getTileByRef(dtTileRef ref) const;
 	
 	/// The maximum number of tiles supported by the navigation mesh.
 	/// @return The maximum number of tiles supported by the navigation mesh.
-	int getMaxTiles() const;
+	NAVMESH_API int getMaxTiles() const;
 	
 	/// Gets the tile at the specified index.
 	///  @param[in]	i		The tile index. [Limit: 0 >= index < #getMaxTiles()]
 	/// @return The tile at the specified index.
-	const dtMeshTile* getTile(int i) const;
+	NAVMESH_API const dtMeshTile* getTile(int i) const;
 
 	/// Gets the tile and polygon for the specified polygon reference.
 	///  @param[in]		ref		The reference for the a polygon.
 	///  @param[out]	tile	The tile containing the polygon.
 	///  @param[out]	poly	The polygon.
 	/// @return The status flags for the operation.
-	dtStatus getTileAndPolyByRef(const dtPolyRef ref, const dtMeshTile** tile, const dtPoly** poly) const;
+	NAVMESH_API dtStatus getTileAndPolyByRef(const dtPolyRef ref, const dtMeshTile** tile, const dtPoly** poly) const;
 	
 	/// Returns the tile and polygon for the specified polygon reference.
 	///  @param[in]		ref		A known valid reference for a polygon.
 	///  @param[out]	tile	The tile containing the polygon.
 	///  @param[out]	poly	The polygon.
-	void getTileAndPolyByRefUnsafe(const dtPolyRef ref, const dtMeshTile** tile, const dtPoly** poly) const;
+	NAVMESH_API void getTileAndPolyByRefUnsafe(const dtPolyRef ref, const dtMeshTile** tile, const dtPoly** poly) const;
 
 	/// Checks the validity of a polygon reference.
 	///  @param[in]	ref		The polygon reference to check.
 	/// @return True if polygon reference is valid for the navigation mesh.
-	bool isValidPolyRef(dtPolyRef ref) const;
+	NAVMESH_API bool isValidPolyRef(dtPolyRef ref) const;
 	
 	/// Gets the polygon reference for the tile's base polygon.
 	///  @param[in]	tile		The tile.
 	/// @return The polygon reference for the base polygon in the specified tile.
-	dtPolyRef getPolyRefBase(const dtMeshTile* tile) const;
+	NAVMESH_API dtPolyRef getPolyRefBase(const dtMeshTile* tile) const;
 	
 	/// Gets the cluster reference for the tile's base cluster.
 	///  @param[in]	tile		The tile.
 	/// @return The cluster reference for the base cluster in the specified tile.
-	dtClusterRef getClusterRefBase(const dtMeshTile* tile) const;
+	NAVMESH_API dtClusterRef getClusterRefBase(const dtMeshTile* tile) const;
 
 	/// Gets the endpoints for an off-mesh connection, ordered by "direction of travel".
 	///  @param[in]		prevRef		The reference of the polygon before the connection.
@@ -637,29 +637,29 @@ public:
 	///  @param[out]	startPos	The start position of the off-mesh connection. [(x, y, z)]
 	///  @param[out]	endPos		The end position of the off-mesh connection. [(x, y, z)]
 	/// @return The status flags for the operation.
-	dtStatus getOffMeshConnectionPolyEndPoints(dtPolyRef prevRef, dtPolyRef polyRef, const dtReal* currentPos, dtReal* startPos, dtReal* endPos) const;
+	NAVMESH_API dtStatus getOffMeshConnectionPolyEndPoints(dtPolyRef prevRef, dtPolyRef polyRef, const dtReal* currentPos, dtReal* startPos, dtReal* endPos) const;
 
 	/// Gets the specified off-mesh connection: point type.
 	///  @param[in]	ref		The polygon reference of the off-mesh connection.
 	/// @return The specified off-mesh connection, or null if the polygon reference is not valid.
-	const dtOffMeshConnection* getOffMeshConnectionByRef(dtPolyRef ref) const;
+	NAVMESH_API const dtOffMeshConnection* getOffMeshConnectionByRef(dtPolyRef ref) const;
 
 	/// Updates area and flags for specified off-mesh connection: point type
 	///  @param[in] userId	User Id of connection
 	///	 @param[in] newArea	Area code to apply
-	void updateOffMeshConnectionByUserId(unsigned long long int userId, unsigned char newArea, unsigned short newFlags);
+	NAVMESH_API void updateOffMeshConnectionByUserId(unsigned long long int userId, unsigned char newArea, unsigned short newFlags);
 
 	//@UE BEGIN
 #if WITH_NAVMESH_SEGMENT_LINKS
 	/// Gets the specified off-mesh connection: segment type
 	///  @param[in]	ref		The polygon reference of the off-mesh connection.
 	/// @return The specified off-mesh connection, or null if the polygon reference is not valid.
-	const dtOffMeshSegmentConnection* getOffMeshSegmentConnectionByRef(dtPolyRef ref) const;
+	NAVMESH_API const dtOffMeshSegmentConnection* getOffMeshSegmentConnectionByRef(dtPolyRef ref) const;
 
 	/// Updates area and flags for specified off-mesh connection: segment type
 	///  @param[in] userId	User Id of connection
 	///	 @param[in] newArea	Area code to apply
-	void updateOffMeshSegmentConnectionByUserId(unsigned int userId, unsigned char newArea, unsigned short newFlags);
+	NAVMESH_API void updateOffMeshSegmentConnectionByUserId(unsigned int userId, unsigned char newArea, unsigned short newFlags);
 #endif // WITH_NAVMESH_SEGMENT_LINKS
 	//@UE END
 
@@ -673,44 +673,44 @@ public:
 	///  @param[in]	ref		The polygon reference.
 	///  @param[in]	flags	The new flags for the polygon.
 	/// @return The status flags for the operation.
-	dtStatus setPolyFlags(dtPolyRef ref, unsigned short flags);
+	NAVMESH_API dtStatus setPolyFlags(dtPolyRef ref, unsigned short flags);
 
 	/// Gets the user defined flags for the specified polygon.
 	///  @param[in]		ref				The polygon reference.
 	///  @param[out]	resultFlags		The polygon flags.
 	/// @return The status flags for the operation.
-	dtStatus getPolyFlags(dtPolyRef ref, unsigned short* resultFlags) const;
+	NAVMESH_API dtStatus getPolyFlags(dtPolyRef ref, unsigned short* resultFlags) const;
 
 	/// Sets the user defined area for the specified polygon.
 	///  @param[in]	ref		The polygon reference.
 	///  @param[in]	area	The new area id for the polygon. [Limit: < #DT_MAX_AREAS]
 	/// @return The status flags for the operation.
-	dtStatus setPolyArea(dtPolyRef ref, unsigned char area);
+	NAVMESH_API dtStatus setPolyArea(dtPolyRef ref, unsigned char area);
 
 	/// Gets the user defined area for the specified polygon.
 	///  @param[in]		ref			The polygon reference.
 	///  @param[out]	resultArea	The area id for the polygon.
 	/// @return The status flags for the operation.
-	dtStatus getPolyArea(dtPolyRef ref, unsigned char* resultArea) const;
+	NAVMESH_API dtStatus getPolyArea(dtPolyRef ref, unsigned char* resultArea) const;
 
 	/// Gets the size of the buffer required by #storeTileState to store the specified tile's state.
 	///  @param[in]	tile	The tile.
 	/// @return The size of the buffer required to store the state.
-	int getTileStateSize(const dtMeshTile* tile) const;
+	NAVMESH_API int getTileStateSize(const dtMeshTile* tile) const;
 	
 	/// Stores the non-structural state of the tile in the specified buffer. (Flags, area ids, etc.)
 	///  @param[in]		tile			The tile.
 	///  @param[out]	data			The buffer to store the tile's state in.
 	///  @param[in]		maxDataSize		The size of the data buffer. [Limit: >= #getTileStateSize]
 	/// @return The status flags for the operation.
-	dtStatus storeTileState(const dtMeshTile* tile, unsigned char* data, const int maxDataSize) const;
+	NAVMESH_API dtStatus storeTileState(const dtMeshTile* tile, unsigned char* data, const int maxDataSize) const;
 	
 	/// Restores the state of the tile.
 	///  @param[in]	tile			The tile.
 	///  @param[in]	data			The new state. (Obtained from #storeTileState.)
 	///  @param[in]	maxDataSize		The size of the state within the data buffer.
 	/// @return The status flags for the operation.
-	dtStatus restoreTileState(dtMeshTile* tile, const unsigned char* data, const int maxDataSize);
+	NAVMESH_API dtStatus restoreTileState(dtMeshTile* tile, const unsigned char* data, const int maxDataSize);
 	
 	/// @}
 
@@ -795,7 +795,7 @@ public:
 #endif // WITH_NAVMESH_CLUSTER_LINKS
 
 	/// Shift navigation mesh by provided offset
-	void applyWorldOffset(const dtReal* offset);
+	NAVMESH_API void applyWorldOffset(const dtReal* offset);
 
 	/// Helper for accessing links
 	inline dtLink& getLink(dtMeshTile* tile, unsigned int linkIdx)
@@ -822,7 +822,7 @@ public:
 #endif // WITH_NAVMESH_CLUSTER_LINKS
 
 	/// Helper for creating links in off-mesh connections
-	void linkOffMeshHelper(dtMeshTile* tile0, unsigned int polyIdx0, dtMeshTile* tile1, unsigned int polyIdx1, unsigned char side, unsigned char edge);
+	NAVMESH_API void linkOffMeshHelper(dtMeshTile* tile0, unsigned int polyIdx0, dtMeshTile* tile1, unsigned int polyIdx1, unsigned char side, unsigned char edge);
 
 	inline bool isEmpty() const
 	{
@@ -835,10 +835,10 @@ public:
 		return m_saltBits;
 	}
 
-	void applyAreaCostOrder(unsigned char* costOrder);
+	NAVMESH_API void applyAreaCostOrder(unsigned char* costOrder);
 	
 	/// Returns neighbour tile count based on side of given tile.
-	int getNeighbourTilesCountAt(const int x, const int y, const int side) const;
+	NAVMESH_API int getNeighbourTilesCountAt(const int x, const int y, const int side) const;
 
 	bool getNeighbourCoords(const int x, const int y, const int side, int& outX, int& outY) const
 	{
@@ -884,67 +884,67 @@ private:
 	};
 
 	/// Returns pointer to tile in the tile array.
-	dtMeshTile* getTile(int i);
+	NAVMESH_API dtMeshTile* getTile(int i);
 
 	/// Returns neighbour tile based on side.
-	int getTilesAt(const int x, const int y,
+	NAVMESH_API int getTilesAt(const int x, const int y,
 				   dtMeshTile** tiles, const int maxTiles) const;
 
 	/// Returns neighbour tile based on side.
-	int getNeighbourTilesAt(const int x, const int y, const int side,
+	NAVMESH_API int getNeighbourTilesAt(const int x, const int y, const int side,
 							dtMeshTile** tiles, const int maxTiles) const;
 	
 	/// [UE] Returns all polygons in neighbour tile based on portal defined by the segment.
-	int findConnectingPolys(const dtReal* va, const dtReal* vb,
+	NAVMESH_API int findConnectingPolys(const dtReal* va, const dtReal* vb,
 		const dtMeshTile* fromTile, int fromPolyIdx,
 		const dtMeshTile* tile, int side,
 		dtChunkArray<FConnectingPolyData>& cons) const;
 
 	/// Builds internal polygons links for a tile.
-	void connectIntLinks(dtMeshTile* tile);
+	NAVMESH_API void connectIntLinks(dtMeshTile* tile);
 	/// Builds internal polygons links for a tile.
-	void baseOffMeshLinks(dtMeshTile* tile);
+	NAVMESH_API void baseOffMeshLinks(dtMeshTile* tile);
 
 	/// Builds external polygon links for a tile.
-	void connectExtLinks(dtMeshTile* tile, dtMeshTile* target, int side, bool updateCLinks);
+	NAVMESH_API void connectExtLinks(dtMeshTile* tile, dtMeshTile* target, int side, bool updateCLinks);
 	/// Builds external polygon links for a tile.
-	void connectExtOffMeshLinks(dtMeshTile* tile, dtMeshTile* target, int side, bool updateCLinks);
+	NAVMESH_API void connectExtOffMeshLinks(dtMeshTile* tile, dtMeshTile* target, int side, bool updateCLinks);
 	
 	/// Removes external links at specified side.
-	void unconnectExtLinks(dtMeshTile* tile, dtMeshTile* target);
+	NAVMESH_API void unconnectExtLinks(dtMeshTile* tile, dtMeshTile* target);
 	
 	//@UE BEGIN
 #if WITH_NAVMESH_CLUSTER_LINKS
 	/// Try to connect clusters
-	void connectClusterLink(dtMeshTile* tile0, unsigned int cluster0,
+	NAVMESH_API void connectClusterLink(dtMeshTile* tile0, unsigned int cluster0,
 							dtMeshTile* tile1, unsigned int cluster1,
 							unsigned char flags, bool bCheckExisting = true);
 	
 	/// Removes cluster links at specified side.
-	void unconnectClusterLinks(dtMeshTile* tile, dtMeshTile* target);
+	NAVMESH_API void unconnectClusterLinks(dtMeshTile* tile, dtMeshTile* target);
 #endif // WITH_NAVMESH_CLUSTER_LINKS
 	//@UE END
 
 	// TODO: These methods are duplicates from dtNavMeshQuery, but are needed for off-mesh connection finding.
 	
 	/// Queries polygons within a tile.
-	int queryPolygonsInTile(const dtMeshTile* tile, const dtReal* qmin, const dtReal* qmax,
+	NAVMESH_API int queryPolygonsInTile(const dtMeshTile* tile, const dtReal* qmin, const dtReal* qmax,
 							dtPolyRef* polys, const int maxPolys, bool bExcludeUnwalkable = false) const;
 	/// Find nearest polygon within a tile.
-	dtPolyRef findNearestPolyInTile(const dtMeshTile* tile, const dtReal* center,
+	NAVMESH_API dtPolyRef findNearestPolyInTile(const dtMeshTile* tile, const dtReal* center,
 									const dtReal* extents, dtReal* nearestPt, bool bExcludeUnwalkable = false) const;
-	dtPolyRef findCheapestNearPolyInTile(const dtMeshTile* tile, const dtReal* center,
+	NAVMESH_API dtPolyRef findCheapestNearPolyInTile(const dtMeshTile* tile, const dtReal* center,
 										 const dtReal* extents, dtReal* nearestPt) const;
 	/// Returns closest point on polygon.
-	void closestPointOnPolyInTile(const dtMeshTile* tile, unsigned int ip,
+	NAVMESH_API void closestPointOnPolyInTile(const dtMeshTile* tile, unsigned int ip,
 								  const dtReal* pos, dtReal* closest) const;
 
 private:
 	/// Calculates whether the tile grid location can fit in the tile indices type (currently an int)
 	///  @param[in]	tx			Tile X coord
 	///  @param[in]	ty			Tile Y coord
-	bool isTileLocInValidRange(const dtReal tx, const dtReal ty) const;
-	void calcTileLoc(const dtReal* pos, dtReal* tx, dtReal* ty) const;
+	NAVMESH_API bool isTileLocInValidRange(const dtReal tx, const dtReal ty) const;
+	NAVMESH_API void calcTileLoc(const dtReal* pos, dtReal* tx, dtReal* ty) const;
 
 public:
 	dtNavMeshParams m_params;			///< Current initialization params.
diff --git a/Engine/Source/Runtime/Navmesh/Public/Detour/DetourNavMeshQuery.h b/Engine/Source/Runtime/Navmesh/Public/Detour/DetourNavMeshQuery.h
--- a/Engine/Source/Runtime/Navmesh/Public/Detour/DetourNavMeshQuery.h
+++ b/Engine/Source/Runtime/Navmesh/Public/Detour/DetourNavMeshQuery.h
@@ -46,7 +46,7 @@
 // Special link filter is custom filter run only for offmesh links with assigned UserId
 // Used by smart navlinks in UE
 //
-struct NAVMESH_API dtQuerySpecialLinkFilter
+struct dtQuerySpecialLinkFilter
 {
 	virtual ~dtQuerySpecialLinkFilter() {}
 
@@ -63,7 +63,7 @@ struct NAVMESH_API dtQuerySpecialLinkFilter
 };
 
 // [UE: moved all filter variables to struct, DO NOT mess with virtual functions here!]
-struct NAVMESH_API dtQueryFilterData
+struct dtQueryFilterData
 {
 	dtReal m_areaCost[DT_MAX_AREAS];		///< Cost per area type. (Used by default implementation.)
 #if WITH_FIXED_AREA_ENTERING_COST
@@ -84,15 +84,15 @@ struct NAVMESH_API dtQueryFilterData
 	bool m_shouldIgnoreClosedNodes;
 	//@UE END
 
-	dtQueryFilterData();
+	NAVMESH_API dtQueryFilterData();
 	
-	bool equals(const dtQueryFilterData* other) const;
-	void copyFrom(const dtQueryFilterData* source);
+	NAVMESH_API bool equals(const dtQueryFilterData* other) const;
+	NAVMESH_API void copyFrom(const dtQueryFilterData* source);
 };
 
 /// Defines polygon filtering and traversal costs for navigation mesh query operations.
 /// @ingroup detour
-class NAVMESH_API dtQueryFilter
+class dtQueryFilter
 {
 protected:
 	dtQueryFilterData data;
@@ -309,7 +309,7 @@ struct dtQueryResultPack
 	dtQueryResultPack(dtPolyRef inRef, dtReal inCost, const dtReal* inPos, unsigned int inFlag);
 };
 
-struct NAVMESH_API dtQueryResult
+struct dtQueryResult
 {
 	inline void reserve(int n) { data.resize(n); data.resize(0); }
 	inline int size() const { return data.size(); }
@@ -318,13 +318,13 @@ struct NAVMESH_API dtQueryResult
 	inline dtReal getCost(int idx) const { return data[idx].cost; }
 	inline const dtReal* getPos(int idx) const { return data[idx].pos; }
 	inline unsigned int getFlag(int idx) const { return data[idx].flag; }
-	void getPos(int idx, dtReal* pos);
+	NAVMESH_API void getPos(int idx, dtReal* pos);
 
-	void copyRefs(dtPolyRef* refs, int nmax);
-	void copyCosts(dtReal* costs, int nmax);
-	void copyPos(dtReal* pos, int nmax);
-	void copyFlags(unsigned char* flags, int nmax);
-	void copyFlags(unsigned int* flags, int nmax);
+	NAVMESH_API void copyRefs(dtPolyRef* refs, int nmax);
+	NAVMESH_API void copyCosts(dtReal* costs, int nmax);
+	NAVMESH_API void copyPos(dtReal* pos, int nmax);
+	NAVMESH_API void copyFlags(unsigned char* flags, int nmax);
+	NAVMESH_API void copyFlags(unsigned int* flags, int nmax);
 
 protected:
 	dtChunkArray<dtQueryResultPack> data;
@@ -334,7 +334,7 @@ protected:
 	inline void setRef(int idx, dtPolyRef ref) { data[idx].ref = ref; }
 	inline void setCost(int idx, dtReal cost) { data[idx].cost = cost; }
 	inline void setFlag(int idx, unsigned int flag) { data[idx].flag = flag; }
-	void setPos(int idx, const dtReal* pos);
+	NAVMESH_API void setPos(int idx, const dtReal* pos);
 
 	friend class dtNavMeshQuery;
 };
@@ -342,21 +342,21 @@ protected:
 /// Provides the ability to perform pathfinding related queries against
 /// a navigation mesh.
 /// @ingroup detour
-class NAVMESH_API dtNavMeshQuery
+class dtNavMeshQuery
 {
 public:
-	dtNavMeshQuery();
-	~dtNavMeshQuery();
+	NAVMESH_API dtNavMeshQuery();
+	NAVMESH_API ~dtNavMeshQuery();
 	
 	/// Initializes the query object.
 	///  @param[in]		nav			Pointer to the dtNavMesh object to use for all queries.
 	///  @param[in]		maxNodes	Maximum number of search nodes. [Limits: 0 < value <= 65536]
 	///  @param[in]		linkFilter	Special link filter used for every query
 	/// @returns The status flags for the query.
-	dtStatus init(const dtNavMesh* nav, const int maxNodes, dtQuerySpecialLinkFilter* linkFilter = 0);
+	NAVMESH_API dtStatus init(const dtNavMesh* nav, const int maxNodes, dtQuerySpecialLinkFilter* linkFilter = 0);
 
 	/// UE: updates special link filter for this query
-	void updateLinkFilter(dtQuerySpecialLinkFilter* linkFilter);
+	NAVMESH_API void updateLinkFilter(dtQuerySpecialLinkFilter* linkFilter);
 	
 	/// @name Standard Pathfinding Functions
 	// /@{
@@ -370,7 +370,7 @@ public:
 	///  @param[in]		filter		The polygon filter to apply to the query.
 	///  @param[out]	result		Results for path corridor, fills in refs and costs for each poly from start to end
 	///	 @param[out]	totalCost			If provided will get filled will total cost of path
-	dtStatus findPath(dtPolyRef startRef, dtPolyRef endRef,
+	NAVMESH_API dtStatus findPath(dtPolyRef startRef, dtPolyRef endRef,
 					  const dtReal* startPos, const dtReal* endPos, const dtReal costLimit, //@UE
 					  const dtQueryFilter* filter,
 					  dtQueryResult& result, dtReal* totalCost) const;
@@ -383,7 +383,7 @@ public:
 	///  @param[out]	result				Fills in positions, refs and flags
 	///  @param[in]		options				Query options. (see: #dtStraightPathOptions)
 	/// @returns The status flags for the query.
-	dtStatus findStraightPath(const dtReal* startPos, const dtReal* endPos,
+	NAVMESH_API dtStatus findStraightPath(const dtReal* startPos, const dtReal* endPos,
 							  const dtPolyRef* path, const int pathSize,
 							  dtQueryResult& result, const int options = 0) const;
 
@@ -404,7 +404,7 @@ public:
 	///  @param[in]		requireNavigableEndLocation	Define if the end location is required to be a valid navmesh polygon //@UE
 	///  @param[in]		filter						The polygon filter to apply to the query.
 	/// @returns The status flags for the query.
-	dtStatus initSlicedFindPath(dtPolyRef startRef, dtPolyRef endRef,
+	NAVMESH_API dtStatus initSlicedFindPath(dtPolyRef startRef, dtPolyRef endRef,
 								const dtReal* startPos, const dtReal* endPos, const dtReal costLimit, const bool requireNavigableEndLocation, //@UE
 								const dtQueryFilter* filter);
 
@@ -412,7 +412,7 @@ public:
 	///  @param[in]		maxIter		The maximum number of iterations to perform.
 	///  @param[out]	doneIters	The actual number of iterations completed. [opt]
 	/// @returns The status flags for the query.
-	dtStatus updateSlicedFindPath(const int maxIter, int* doneIters);
+	NAVMESH_API dtStatus updateSlicedFindPath(const int maxIter, int* doneIters);
 
 	/// Finalizes and returns the results of a sliced path query.
 	///  @param[out]	path		An ordered list of polygon references representing the path. (Start to end.) 
@@ -420,7 +420,7 @@ public:
 	///  @param[out]	pathCount	The number of polygons returned in the @p path array.
 	///  @param[in]		maxPath		The max number of polygons the path array can hold. [Limit: >= 1]
 	/// @returns The status flags for the query.
-	dtStatus finalizeSlicedFindPath(dtPolyRef* path, int* pathCount, const int maxPath);
+	NAVMESH_API dtStatus finalizeSlicedFindPath(dtPolyRef* path, int* pathCount, const int maxPath);
 	
 	/// Finalizes and returns the results of an incomplete sliced path query, returning the path to the furthest
 	/// polygon on the existing path that was visited during the search.
@@ -431,7 +431,7 @@ public:
 	///  @param[out]	pathCount		The number of polygons returned in the @p path array.
 	///  @param[in]		maxPath			The max number of polygons the @p path array can hold. [Limit: >= 1]
 	/// @returns The status flags for the query.
-	dtStatus finalizeSlicedFindPathPartial(const dtPolyRef* existing, const int existingSize,
+	NAVMESH_API dtStatus finalizeSlicedFindPathPartial(const dtPolyRef* existing, const int existingSize,
 										   dtPolyRef* path, int* pathCount, const int maxPath);
 
 	///@}
@@ -450,7 +450,7 @@ public:
 	///  @param[out]	resultCount		The number of polygons found.
 	///  @param[in]		maxResult		The maximum number of polygons the result arrays can hold.
 	/// @returns The status flags for the query.
-	dtStatus findPolysAroundCircle(dtPolyRef startRef, const dtReal* centerPos, const dtReal radius,
+	NAVMESH_API dtStatus findPolysAroundCircle(dtPolyRef startRef, const dtReal* centerPos, const dtReal radius,
 								   const dtQueryFilter* filter,
 								   dtPolyRef* resultRef, dtPolyRef* resultParent, dtReal* resultCost,
 								   int* resultCount, const int maxResult) const;
@@ -468,7 +468,7 @@ public:
 	///  @param[out]	resultCount		The number of polygons found.
 	///  @param[in]		maxResult		The maximum number of polygons the result arrays can hold.
 	/// @returns The status flags for the query.
-	dtStatus findPolysAroundShape(dtPolyRef startRef, const dtReal* verts, const int nverts,
+	NAVMESH_API dtStatus findPolysAroundShape(dtPolyRef startRef, const dtReal* verts, const int nverts,
 								  const dtQueryFilter* filter,
 								  dtPolyRef* resultRef, dtPolyRef* resultParent, dtReal* resultCost,
 								  int* resultCount, const int maxResult) const;
@@ -483,7 +483,7 @@ public:
 	///  @param[out]	resultCount		The number of polygons found. [opt]
 	///  @param[in]		maxResult		The maximum number of polygons the result arrays can hold.
 	/// @returns The status flags for the querydtReal
-	dtStatus findPolysInPathDistance(dtPolyRef startRef, const dtReal* centerPos, const dtReal pathDistance,
+	NAVMESH_API dtStatus findPolysInPathDistance(dtPolyRef startRef, const dtReal* centerPos, const dtReal pathDistance,
 									const dtQueryFilter* filter, dtPolyRef* resultRef,
 									int* resultCount, const int maxResult) const;
 
@@ -499,7 +499,7 @@ public:
 	///  @param[out]	nearestPt	The nearest point on the polygon. [opt] [(x, y, z)]
 	///  @param[in]		referencePt	If supplied replaces @param center in terms of distance measurements. [opt] [(x, y, z)]
 	/// @returns The status flags for the query.
-	dtStatus findNearestPoly(const dtReal* center, const dtReal* extents,
+	NAVMESH_API dtStatus findNearestPoly(const dtReal* center, const dtReal* extents,
 							 const dtQueryFilter* filter,
 							 dtPolyRef* nearestRef, dtReal* nearestPt, const dtReal* referencePt = 0) const;
 
@@ -512,7 +512,7 @@ public:
 	///  @param[in]		referencePt	If supplied replaces @param center in terms of distance measurements. [opt] [(x, y, z)]
 	///  @param[in]		tolerance	Radius around best 2D point for picking vertical match
 	/// @returns The status flags for the query.
-	dtStatus findNearestPoly2D(const dtReal* center, const dtReal* extents,
+	NAVMESH_API dtStatus findNearestPoly2D(const dtReal* center, const dtReal* extents,
 							const dtQueryFilter* filter,
 							dtPolyRef* outProjectedRef, dtReal* outProjectedPt,
 							const dtReal* referencePt = 0, dtReal tolerance = 0) const;
@@ -524,7 +524,7 @@ public:
 	///  @param[out]	nearestRef	The reference id of the nearest polygon.
 	///  @param[out]	nearestPt	The nearest point on the polygon. [opt] [(x, y, z)]
 	/// @returns The status flags for the query.
-	dtStatus findNearestContainingPoly(const dtReal* center, const dtReal* extents,
+	NAVMESH_API dtStatus findNearestContainingPoly(const dtReal* center, const dtReal* extents,
 									   const dtQueryFilter* filter,
 									   dtPolyRef* nearestRef, dtReal* nearestPt) const;
 	//@UE END
@@ -537,7 +537,7 @@ public:
 	///  @param[out]	polyCount	The number of polygons in the search result.
 	///  @param[in]		maxPolys	The maximum number of polygons the search result can hold.
 	/// @returns The status flags for the query.
-	dtStatus queryPolygons(const dtReal* center, const dtReal* extents,
+	NAVMESH_API dtStatus queryPolygons(const dtReal* center, const dtReal* extents,
 						   const dtQueryFilter* filter,
 						   dtPolyRef* polys, int* polyCount, const int maxPolys) const;
 
@@ -552,26 +552,26 @@ public:
 	///  @param[out]	resultCount		The number of polygons found.
 	///  @param[in]		maxResult		The maximum number of polygons the result arrays can hold.
 	/// @returns The status flags for the query.
-	dtStatus findLocalNeighbourhood(dtPolyRef startRef, const dtReal* centerPos, const dtReal radius,
+	NAVMESH_API dtStatus findLocalNeighbourhood(dtPolyRef startRef, const dtReal* centerPos, const dtReal radius,
 									const dtQueryFilter* filter,
 									dtPolyRef* resultRef, dtPolyRef* resultParent,
 									int* resultCount, const int maxResult) const;
 
 	/// [UE] Finds the wall segments in local neighbourhood
-	dtStatus findWallsInNeighbourhood(dtPolyRef startRef, const dtReal* centerPos, const dtReal radius,
+	NAVMESH_API dtStatus findWallsInNeighbourhood(dtPolyRef startRef, const dtReal* centerPos, const dtReal radius,
 									  const dtQueryFilter* filter, 
 									  dtPolyRef* neiRefs, int* neiCount, const int maxNei,
 									  dtReal* resultWalls, dtPolyRef* resultRefs, int* resultCount, const int maxResult) const;
 
 
 	/// [UE] Finds the wall segments that overlap the polygon shape.
-	dtStatus findWallsOverlappingShape(dtPolyRef startRef, const dtReal* verts, const int nverts,
+	NAVMESH_API dtStatus findWallsOverlappingShape(dtPolyRef startRef, const dtReal* verts, const int nverts,
 									   const dtQueryFilter* filter, 
 									   dtPolyRef* neiRefs, int* neiCount, const int maxNei,
 									   dtReal* resultWalls, dtPolyRef* resultRefs, int* resultCount, const int maxResult) const;
 
 	/// [UE] Finds the wall segments that overlap the polygon shape.
-	dtStatus findWallsAroundPath(const dtPolyRef* path, const int pathCount, const dtReal* searchAreaPoly, const int searchAreaPolyCount,
+	NAVMESH_API dtStatus findWallsAroundPath(const dtPolyRef* path, const int pathCount, const dtReal* searchAreaPoly, const int searchAreaPolyCount,
 									   const dtReal maxAreaEnterCost, const dtQueryFilter* filter,
 									   dtPolyRef* neiRefs, int* neiCount, const int maxNei,
 									   dtReal* resultWalls, dtPolyRef* resultRefs, int* resultCount, const int maxResult) const;
@@ -586,7 +586,7 @@ public:
 	///  @param[out]	visitedCount	The number of polygons visited during the move.
 	///  @param[in]		maxVisitedSize	The maximum number of polygons the @p visited array can hold.
 	/// @returns The status flags for the query.
-	dtStatus moveAlongSurface(dtPolyRef startRef, const dtReal* startPos, const dtReal* endPos,
+	NAVMESH_API dtStatus moveAlongSurface(dtPolyRef startRef, const dtReal* startPos, const dtReal* endPos,
 							  const dtQueryFilter* filter,
 							  dtReal* resultPos, dtPolyRef* visited, int* visitedCount, const int maxVisitedSize) const;
 	
@@ -604,7 +604,7 @@ public:
 	///  @param[in]		maxPath		The maximum number of polygons the @p path array can hold.
 	///  @param[in]		walkableArea Specific area that can be visited or 255 to skip it and test all areas.
 	/// @returns The status flags for the query.
-	dtStatus raycast(dtPolyRef startRef, const dtReal* startPos, const dtReal* endPos,
+	NAVMESH_API dtStatus raycast(dtPolyRef startRef, const dtReal* startPos, const dtReal* endPos,
 					 const dtQueryFilter* filter,
 					 dtReal* t, dtReal* hitNormal, dtPolyRef* path, int* pathCount, const int maxPath) const;
 	
@@ -618,7 +618,7 @@ public:
 	///  @param[out]	hitNormal		The normalized ray formed from the wall point to the 
 	///  								source point. [(x, y, z)]
 	/// @returns The status flags for the query.
-	dtStatus findDistanceToWall(dtPolyRef startRef, const dtReal* centerPos, const dtReal maxRadius,
+	NAVMESH_API dtStatus findDistanceToWall(dtPolyRef startRef, const dtReal* centerPos, const dtReal maxRadius,
 								const dtQueryFilter* filter,
 								dtReal* hitDist, dtReal* hitPos, dtReal* hitNormal) const;
 	
@@ -631,7 +631,7 @@ public:
 	///  @param[out]	segmentCount	The number of segments returned.
 	///  @param[in]		maxSegments		The maximum number of segments the result arrays can hold.
 	/// @returns The status flags for the query.
-	dtStatus getPolyWallSegments(dtPolyRef ref, const dtQueryFilter* filter,
+	NAVMESH_API dtStatus getPolyWallSegments(dtPolyRef ref, const dtQueryFilter* filter,
 								 dtReal* segmentVerts, dtPolyRef* segmentRefs, int* segmentCount,
 								 const int maxSegments) const;
 
@@ -642,7 +642,7 @@ public:
 	///  @param[out]	randomRef		The reference id of the random location.
 	///  @param[out]	randomPt		The random location. 
 	/// @returns The status flags for the query.
-	dtStatus findRandomPoint(const dtQueryFilter* filter, float(*frand)(),
+	NAVMESH_API dtStatus findRandomPoint(const dtQueryFilter* filter, float(*frand)(),
 							 dtPolyRef* randomRef, dtReal* randomPt) const;
 
 	/// Returns random location on navmesh within the reach of specified location.
@@ -655,7 +655,7 @@ public:
 	///  @param[out]	randomRef		The reference id of the random location.
 	///  @param[out]	randomPt		The random location. [(x, y, z)]
 	/// @returns The status flags for the query.
-	dtStatus findRandomPointAroundCircle(dtPolyRef startRef, const dtReal* centerPos, const dtReal maxRadius,
+	NAVMESH_API dtStatus findRandomPointAroundCircle(dtPolyRef startRef, const dtReal* centerPos, const dtReal maxRadius,
 										 const dtQueryFilter* filter, float(*frand)(),
 										 dtPolyRef* randomRef, dtReal* randomPt) const;
 
@@ -664,7 +664,7 @@ public:
 	///  @param[in]		frand			Function returning a random number [0..1).
 	///  @param[out]	randomPt		The random location. [(x, y, z)]
 	/// @returns The status flags for the query.
-	dtStatus findRandomPointInPoly(dtPolyRef ref, float(*frand)(), dtReal* randomPt) const;
+	NAVMESH_API dtStatus findRandomPointInPoly(dtPolyRef ref, float(*frand)(), dtReal* randomPt) const;
 
 	//@UE BEGIN
 #if WITH_NAVMESH_CLUSTER_LINKS
@@ -672,21 +672,21 @@ public:
 	/// (cheap, does not care about link costs)
 	///  @param[in]		startRef			The reference id of the start polygon.
 	///  @param[in]		endRef				The reference id of the end polygon.
-	dtStatus testClusterPath(dtPolyRef startRef, dtPolyRef endRef) const; 
+	NAVMESH_API dtStatus testClusterPath(dtPolyRef startRef, dtPolyRef endRef) const; 
 
 	/// Returns random location on navmesh within specified cluster.
 	///  @param[in]		frand			Function returning a random number [0..1).
 	///  @param[out]	randomRef		The reference id of the random location.
 	///  @param[out]	randomPt		The random location. 
 	/// @returns The status flags for the query.
-	dtStatus findRandomPointInCluster(dtClusterRef clusterRef, float(*frand)(),
+	NAVMESH_API dtStatus findRandomPointInCluster(dtClusterRef clusterRef, float(*frand)(),
 									  dtPolyRef* randomRef, dtReal* randomPt) const;
 
 	/// Gets the cluster containing given polygon
 	///  @param[in]		polyRef		The reference id of the polygon.
 	///  @param[out]	clusterRef	The reference id of the cluster
 	/// @returns The status flags for the query.
-	dtStatus getPolyCluster(dtPolyRef polyRef, dtClusterRef& clusterRef) const;
+	NAVMESH_API dtStatus getPolyCluster(dtPolyRef polyRef, dtClusterRef& clusterRef) const;
 #endif // WITH_NAVMESH_CLUSTER_LINKS
 //@UE END
 	
@@ -695,7 +695,7 @@ public:
 	///  @param[in]		pos			The position to check. [(x, y, z)]
 	///  @param[out]	closest		The closest point on the polygon. [(x, y, z)]
 	/// @returns The status flags for the query.
-	dtStatus closestPointOnPoly(dtPolyRef ref, const dtReal* pos, dtReal* closest) const;
+	NAVMESH_API dtStatus closestPointOnPoly(dtPolyRef ref, const dtReal* pos, dtReal* closest) const;
 	
 	/// Returns a point on the boundary closest to the source point if the source point is outside the 
 	/// polygon's xz-bounds.
@@ -703,28 +703,28 @@ public:
 	///  @param[in]		pos			The position to check. [(x, y, z)]
 	///  @param[out]	closest		The closest point. [(x, y, z)]
 	/// @returns The status flags for the query.
-	dtStatus closestPointOnPolyBoundary(dtPolyRef ref, const dtReal* pos, dtReal* closest) const;
+	NAVMESH_API dtStatus closestPointOnPolyBoundary(dtPolyRef ref, const dtReal* pos, dtReal* closest) const;
 
 	/// Finds the point's projection on the specified polygon.
 	///  @param[in]		ref			The reference id of the polygon.
 	///  @param[in]		pos			The position to check. [(x, y, z)]
 	///  @param[out]	closest		The projected point on the polygon. [(x, y, z)]
 	/// @returns The status flags for the query.
-	dtStatus projectedPointOnPoly(dtPolyRef ref, const dtReal* pos, dtReal* projected) const;
+	NAVMESH_API dtStatus projectedPointOnPoly(dtPolyRef ref, const dtReal* pos, dtReal* projected) const;
 	
 	/// Checks if specified pos is inside given polygon specified by ref
 	///  @param[in]		ref			The reference id of the polygon.
 	///  @param[in]		pos			The position to check. [(x, y, z)]
 	///  @param[out]	result		The result of the check, whether the point is inside (true) or not (false)
 	/// @returns The status flags for the query.
-	dtStatus isPointInsidePoly(dtPolyRef ref, const dtReal* pos, bool& result) const;
+	NAVMESH_API dtStatus isPointInsidePoly(dtPolyRef ref, const dtReal* pos, bool& result) const;
 
 	/// Gets the height of the polygon at the provided position using the height detail. (Most accurate.)
 	///  @param[in]		ref			The reference id of the polygon.
 	///  @param[in]		pos			A position within the xz-bounds of the polygon. [(x, y, z)]
 	///  @param[out]	height		The height at the surface of the polygon.
 	/// @returns The status flags for the query.
-	dtStatus getPolyHeight(dtPolyRef ref, const dtReal* pos, dtReal* height) const;
+	NAVMESH_API dtStatus getPolyHeight(dtPolyRef ref, const dtReal* pos, dtReal* height) const;
 
 	/// @}
 	/// @name Miscellaneous Functions
@@ -733,18 +733,18 @@ public:
 	/// Returns true if the polygon reference is valid and passes the filter restrictions.
 	///  @param[in]		ref			The polygon reference to check.
 	///  @param[in]		filter		The filter to apply.
-	bool isValidPolyRef(dtPolyRef ref, const dtQueryFilter* filter) const;
+	NAVMESH_API bool isValidPolyRef(dtPolyRef ref, const dtQueryFilter* filter) const;
 
 	/// Returns true if the polygon reference is in the closed list. 
 	///  @param[in]		ref		The reference id of the polygon to check.
 	/// @returns True if the polygon is in closed list.
-	bool isInClosedList(dtPolyRef ref) const;
+	NAVMESH_API bool isInClosedList(dtPolyRef ref) const;
 	
 	/// Returns true if the cluster link was used in previous search. 
 	///  @param[in]		cFrom		The reference id of the start cluster.
 	///  @param[in]		cto			The reference id of the end cluster.
 	/// @returns True if the cluster link was searched.
-	bool wasClusterLinkSearched(dtPolyRef cFrom, dtPolyRef cTo) const;
+	NAVMESH_API bool wasClusterLinkSearched(dtPolyRef cFrom, dtPolyRef cTo) const;
 
 	/// Gets the node pool.
 	/// @returns The node pool.
@@ -763,43 +763,43 @@ public:
 private:
 	
 	/// Returns neighbour tile based on side.
-	dtMeshTile* getNeighbourTileAt(int x, int y, int side) const;
+	NAVMESH_API dtMeshTile* getNeighbourTileAt(int x, int y, int side) const;
 
 	/// Queries polygons within a tile.
-	int queryPolygonsInTile(const dtMeshTile* tile, const dtReal* qmin, const dtReal* qmax, const dtQueryFilter* filter,
+	NAVMESH_API int queryPolygonsInTile(const dtMeshTile* tile, const dtReal* qmin, const dtReal* qmax, const dtQueryFilter* filter,
 							dtPolyRef* polys, const int maxPolys) const;
 	/// Find nearest polygon within a tile.
-	dtPolyRef findNearestPolyInTile(const dtMeshTile* tile, const dtReal* center, const dtReal* extents,
+	NAVMESH_API dtPolyRef findNearestPolyInTile(const dtMeshTile* tile, const dtReal* center, const dtReal* extents,
 									const dtQueryFilter* filter, dtReal* nearestPt) const;
 	/// Returns closest point on polygon.
-	void closestPointOnPolyInTile(const dtMeshTile* tile, const dtPoly* poly, const dtReal* pos, dtReal* closest) const;
+	NAVMESH_API void closestPointOnPolyInTile(const dtMeshTile* tile, const dtPoly* poly, const dtReal* pos, dtReal* closest) const;
 
 	/// Returns projected point on polygon.
-	dtStatus projectedPointOnPolyInTile(const dtMeshTile* tile, const dtPoly* poly, const dtReal* pos, dtReal* projected) const;
+	NAVMESH_API dtStatus projectedPointOnPolyInTile(const dtMeshTile* tile, const dtPoly* poly, const dtReal* pos, dtReal* projected) const;
 	
 	//@UE BEGIN
 	// exposing function to be able to generate navigation corridors as sequence of point pairs
 public:
 	/// Returns portal points between two polygons.
-	dtStatus getPortalPoints(dtPolyRef from, dtPolyRef to, dtReal* left, dtReal* right,
+	NAVMESH_API dtStatus getPortalPoints(dtPolyRef from, dtPolyRef to, dtReal* left, dtReal* right,
 							 unsigned char& fromType, unsigned char& toType) const;
-	dtStatus getPortalPoints(dtPolyRef from, const dtPoly* fromPoly, const dtMeshTile* fromTile,
+	NAVMESH_API dtStatus getPortalPoints(dtPolyRef from, const dtPoly* fromPoly, const dtMeshTile* fromTile,
 							 dtPolyRef to, const dtPoly* toPoly, const dtMeshTile* toTile,
 							 dtReal* left, dtReal* right) const;
 	
 	/// Returns edge mid point between two polygons.
-	dtStatus getEdgeMidPoint(dtPolyRef from, dtPolyRef to, dtReal* mid) const;
-	dtStatus getEdgeMidPoint(dtPolyRef from, const dtPoly* fromPoly, const dtMeshTile* fromTile,
+	NAVMESH_API dtStatus getEdgeMidPoint(dtPolyRef from, dtPolyRef to, dtReal* mid) const;
+	NAVMESH_API dtStatus getEdgeMidPoint(dtPolyRef from, const dtPoly* fromPoly, const dtMeshTile* fromTile,
 							 dtPolyRef to, const dtPoly* toPoly, const dtMeshTile* toTile,
 							 dtReal* mid) const;
 	//@UE END
 
 	// Appends vertex to a straight path
-	dtStatus appendVertex(const dtReal* pos, const unsigned char flags, const dtPolyRef ref,
+	NAVMESH_API dtStatus appendVertex(const dtReal* pos, const unsigned char flags, const dtPolyRef ref,
 						  dtQueryResult& result, const bool bOverrideIdenticalPosition = true) const;
 
 	// Appends intermediate portal points to a straight path.
-	dtStatus appendPortals(const int startIdx, const int endIdx, const dtReal* endPos, const dtPolyRef* path,
+	NAVMESH_API dtStatus appendPortals(const int startIdx, const int endIdx, const dtReal* endPos, const dtPolyRef* path,
 						   dtQueryResult& result, const int options) const;
 	
 	inline bool passLinkFilterByRef(const dtMeshTile* tile, const dtPolyRef ref) const
diff --git a/Engine/Source/Runtime/Navmesh/Public/DetourCrowd/DetourCrowd.h b/Engine/Source/Runtime/Navmesh/Public/DetourCrowd/DetourCrowd.h
--- a/Engine/Source/Runtime/Navmesh/Public/DetourCrowd/DetourCrowd.h
+++ b/Engine/Source/Runtime/Navmesh/Public/DetourCrowd/DetourCrowd.h
@@ -137,7 +137,7 @@ enum MoveRequestState
 
 /// Represents an agent managed by a #dtCrowd object.
 /// @ingroup crowd
-struct NAVMESH_API dtCrowdAgent
+struct dtCrowdAgent
 {
 	/// The path corridor the agent is using.
 	dtPathCorridor corridor;
@@ -193,7 +193,7 @@ struct NAVMESH_API dtCrowdAgent
 	unsigned char state;
 };
 
-struct NAVMESH_API dtCrowdAgentAnimation
+struct dtCrowdAgentAnimation
 {
 	dtReal initPos[3], startPos[3], endPos[3];
 	dtPolyRef polyRef;
@@ -222,7 +222,7 @@ enum CrowdBoundaryFlags
 	DT_CROWD_BOUNDARY_IGNORE = 1 << 0,
 };
 
-struct NAVMESH_API dtCrowdAgentDebugInfo
+struct dtCrowdAgentDebugInfo
 {
 	int idx;
 	dtReal optStart[3], optEnd[3];
@@ -232,7 +232,7 @@ struct NAVMESH_API dtCrowdAgentDebugInfo
 
 /// Provides local steering behaviors for a group of agents. 
 /// @ingroup crowd
-class NAVMESH_API dtCrowd
+class dtCrowd
 {
 	int m_maxAgents;
 	int m_numActiveAgents;
@@ -278,49 +278,49 @@ class NAVMESH_API dtCrowd
 	// [UE] if set, crowd agents will use early reach test
 	bool m_earlyReachTest;
 
-	void updateTopologyOptimization(dtCrowdAgent** agents, const int nagents, const dtReal dt);
-	void updateMoveRequest(const dtReal dt);
-	void checkPathValidity(dtCrowdAgent** agents, const int nagents, const dtReal dt);
+	NAVMESH_API void updateTopologyOptimization(dtCrowdAgent** agents, const int nagents, const dtReal dt);
+	NAVMESH_API void updateMoveRequest(const dtReal dt);
+	NAVMESH_API void checkPathValidity(dtCrowdAgent** agents, const int nagents, const dtReal dt);
 
-	bool requestMoveTargetReplan(const int idx, dtPolyRef ref, const dtReal* pos);
+	NAVMESH_API bool requestMoveTargetReplan(const int idx, dtPolyRef ref, const dtReal* pos);
 
-	void purge();
+	NAVMESH_API void purge();
 	
 public:
-	dtCrowd();
-	~dtCrowd();
+	NAVMESH_API dtCrowd();
+	NAVMESH_API ~dtCrowd();
 	
 	/// Initializes the crowd.  
 	///  @param[in]		maxAgents		The maximum number of agents the crowd can manage. [Limit: >= 1]
 	///  @param[in]		maxAgentRadius	The maximum radius of any agent that will be added to the crowd. [Limit: > 0]
 	///  @param[in]		nav				The navigation mesh to use for planning.
 	/// @return True if the initialization succeeded.
-	bool init(const int maxAgents, const dtReal maxAgentRadius, dtNavMesh* nav);
+	NAVMESH_API bool init(const int maxAgents, const dtReal maxAgentRadius, dtNavMesh* nav);
 
 	/// [UE] Initializes the avoidance query.  
 	///  @param[in]		maxNeighbors		The maximum number of processed neighbors
 	///  @param[in]		maxWalls			The maximum number of processed wall segments
 	///  @param[in]		maxCustomPatterns	The maximum number of custom sampling patterns
 	/// @return True if the initialization succeeded.
-	bool initAvoidance(const int maxNeighbors, const int maxWalls, const int maxCustomPatterns);
+	NAVMESH_API bool initAvoidance(const int maxNeighbors, const int maxWalls, const int maxCustomPatterns);
 	
 	/// Sets the shared avoidance configuration for the specified index.
 	///  @param[in]		idx		The index. [Limits: 0 <= value < #DT_CROWD_MAX_OBSTAVOIDANCE_PARAMS]
 	///  @param[in]		params	The new configuration.
-	void setObstacleAvoidanceParams(const int idx, const dtObstacleAvoidanceParams* params);
+	NAVMESH_API void setObstacleAvoidanceParams(const int idx, const dtObstacleAvoidanceParams* params);
 
 	/// Gets the shared avoidance configuration for the specified index.
 	///  @param[in]		idx		The index of the configuration to retreive. 
 	///							[Limits:  0 <= value < #DT_CROWD_MAX_OBSTAVOIDANCE_PARAMS]
 	/// @return The requested configuration.
-	const dtObstacleAvoidanceParams* getObstacleAvoidanceParams(const int idx) const;
+	NAVMESH_API const dtObstacleAvoidanceParams* getObstacleAvoidanceParams(const int idx) const;
 	
 	/// [UE] Sets the shared avoidance sampling pattern for the specified index.
 	///  @param[in]		idx			The index.
 	///  @param[in]		angles		radians from direction of desired velocity [Count: nsamples]
 	///  @param[in]		radii		normalized radii (0..1) for each sample [Count: nsamples]
 	///  @param[in]		nsamples	The number of samples
-	void setObstacleAvoidancePattern(int idx, const dtReal* angles, const dtReal* radii, int nsamples);
+	NAVMESH_API void setObstacleAvoidancePattern(int idx, const dtReal* angles, const dtReal* radii, int nsamples);
 
 	/// [UE] Gets the shared avoidance sampling pattern for the specified index.
 	///  @param[in]		idx			The index.
@@ -328,137 +328,137 @@ public:
 	///  @param[in]		radii		normalized radii (0..1) for each sample [Count: nsamples]
 	///  @param[in]		nsamples	The number of samples
 	/// @return true if pattern was found
-	bool getObstacleAvoidancePattern(int idx, dtReal* angles, dtReal* radii, int* nsamples);
+	NAVMESH_API bool getObstacleAvoidancePattern(int idx, dtReal* angles, dtReal* radii, int* nsamples);
 
 	/// Gets the specified agent from the pool.
 	///	 @param[in]		idx		The agent index. [Limits: 0 <= value < #getAgentCount()]
 	/// @return The requested agent.
-	const dtCrowdAgent* getAgent(const int idx);
+	NAVMESH_API const dtCrowdAgent* getAgent(const int idx);
 
 	/// The maximum number of agents that can be managed by the object.
 	/// @return The maximum number of agents.
-	const int getAgentCount() const;
+	NAVMESH_API const int getAgentCount() const;
 	
 	/// Adds a new agent to the crowd.
 	///  @param[in]		pos		The requested position of the agent. [(x, y, z)]
 	///  @param[in]		params	The configutation of the agent.
 	///  @param[in]		filter	[UE] query filter used by agent
 	/// @return The index of the agent in the agent pool. Or -1 if the agent could not be added.
-	int addAgent(const dtReal* pos, const dtCrowdAgentParams& params, const dtQueryFilter* filter);
+	NAVMESH_API int addAgent(const dtReal* pos, const dtCrowdAgentParams& params, const dtQueryFilter* filter);
 
 	/// Updates the specified agent's configuration.
 	///  @param[in]		idx		The agent index. [Limits: 0 <= value < #getAgentCount()]
 	///  @param[in]		params	The new agent configuration.
-	void updateAgentParameters(const int idx, const dtCrowdAgentParams& params);
+	NAVMESH_API void updateAgentParameters(const int idx, const dtCrowdAgentParams& params);
 
 	/// [UE] Updates the specified agent's query filter.
 	///  @param[in]		idx		The agent index. [Limits: 0 <= value < #getAgentCount()]
 	///  @param[in]		filter	The new filter.
 	/// @return True if the request was successfully submitted.
-	bool updateAgentFilter(const int idx, const dtQueryFilter* filter);
+	NAVMESH_API bool updateAgentFilter(const int idx, const dtQueryFilter* filter);
 
 	/// [UE] Refresh state of agent, used after completing movement through offmesh links
 	///  @param[in]		idx		The agent index. [Limits: 0 <= value < #getAgentCount()]
 	///  @param[in]		repath	If set, agent will invalidate its path
-	void updateAgentState(const int idx, bool repath);
+	NAVMESH_API void updateAgentState(const int idx, bool repath);
 
 	/// Removes the agent from the crowd.
 	///  @param[in]		idx		The agent index. [Limits: 0 <= value < #getAgentCount()]
-	void removeAgent(const int idx);
+	NAVMESH_API void removeAgent(const int idx);
 	
 	/// Submits a new move request for the specified agent.
 	///  @param[in]		idx		The agent index. [Limits: 0 <= value < #getAgentCount()]
 	///  @param[in]		ref		The position's polygon reference.
 	///  @param[in]		pos		The position within the polygon. [(x, y, z)]
 	/// @return True if the request was successfully submitted.
-	bool requestMoveTarget(const int idx, dtPolyRef ref, const dtReal* pos);
+	NAVMESH_API bool requestMoveTarget(const int idx, dtPolyRef ref, const dtReal* pos);
 
 	/// Submits a new move request for the specified agent.
 	///  @param[in]		idx		The agent index. [Limits: 0 <= value < #getAgentCount()]
 	///  @param[in]		vel		The movement velocity. [(x, y, z)]
 	/// @return True if the request was successfully submitted.
-	bool requestMoveVelocity(const int idx, const dtReal* vel);
+	NAVMESH_API bool requestMoveVelocity(const int idx, const dtReal* vel);
 
 	/// Resets any request for the specified agent.
 	///  @param[in]		idx		The agent index. [Limits: 0 <= value < #getAgentCount()]
 	/// @return True if the request was successfully reseted.
-	bool resetMoveTarget(const int idx);
+	NAVMESH_API bool resetMoveTarget(const int idx);
 
 	/// [UE] Switch to waiting state
 	///  @param[in]		idx		The agent index. [Limits: 0 <= value < #getAgentCount()]
-	bool setAgentWaiting(const int idx);
+	NAVMESH_API bool setAgentWaiting(const int idx);
 
 	/// [UE] Switch to offmesh link state
 	///  @param[in]		idx		The agent index. [Limits: 0 <= value < #getAgentCount()]
-	bool setAgentBackOnLink(const int idx);
+	NAVMESH_API bool setAgentBackOnLink(const int idx);
 
 	/// [UE] Resets agent's velocity
 	///  @param[in]		idx		The agent index. [Limits: 0 <= value < #getAgentCount()]
-	bool resetAgentVelocity(const int idx);
+	NAVMESH_API bool resetAgentVelocity(const int idx);
 
 	/// Gets the active agents int the agent pool.
 	///  @param[out]	agents		An array of agent pointers. [(#dtCrowdAgent *) * maxAgents]
 	///  @param[in]		maxAgents	The size of the crowd agent array.
 	/// @return The number of agents returned in @p agents.
-	int getActiveAgents(dtCrowdAgent** agents, const int maxAgents);
+	NAVMESH_API int getActiveAgents(dtCrowdAgent** agents, const int maxAgents);
 
 	/// Cache list of active agents
 	/// @return The number of active agents
-	int cacheActiveAgents();
+	NAVMESH_API int cacheActiveAgents();
 
 	/// Updates the steering and positions of all agents.
 	///  @param[in]		dt		The time, in seconds, to update the simulation. [Limit: > 0]
 	///  @param[out]	debug	A debug object to load with debug information. [Opt]
-	void update(const dtReal dt, dtCrowdAgentDebugInfo* debug);
+	NAVMESH_API void update(const dtReal dt, dtCrowdAgentDebugInfo* debug);
 	
 	/// [UE] Split update into several smaller components: path validity, path cache and path optimizations
 	/// @param[in]		dt		Delta time in seconds
 	/// @param[in]		nagents	Number of active agents
-	void updateStepPaths(const dtReal dt, dtCrowdAgentDebugInfo* debug);
+	NAVMESH_API void updateStepPaths(const dtReal dt, dtCrowdAgentDebugInfo* debug);
 
 	/// [UE] Split update into several smaller components: neighbors and boundaries
 	/// @param[in]		dt		Delta time in seconds
 	/// @param[in]		nagents	Number of active agents
-	void updateStepProximityData(const dtReal dt, dtCrowdAgentDebugInfo* debug);
+	NAVMESH_API void updateStepProximityData(const dtReal dt, dtCrowdAgentDebugInfo* debug);
 
 	/// [UE] Split update into several smaller components: next corner for move, trigger offmesh links
 	/// @param[in]		dt		Delta time in seconds
 	/// @param[in]		nagents	Number of active agents
-	void updateStepNextMovePoint(const dtReal dt, dtCrowdAgentDebugInfo* debug);
+	NAVMESH_API void updateStepNextMovePoint(const dtReal dt, dtCrowdAgentDebugInfo* debug);
 
 	/// [UE] Split update into several smaller components: steering
 	/// @param[in]		dt		Delta time in seconds
 	/// @param[in]		nagents	Number of active agents
-	void updateStepSteering(const dtReal dt, dtCrowdAgentDebugInfo* debug);
+	NAVMESH_API void updateStepSteering(const dtReal dt, dtCrowdAgentDebugInfo* debug);
 
 	/// [UE] Split update into several smaller components: avoidance
 	/// @param[in]		dt		Delta time in seconds
 	/// @param[in]		nagents	Number of active agents
-	void updateStepAvoidance(const dtReal dt, dtCrowdAgentDebugInfo* debug);
+	NAVMESH_API void updateStepAvoidance(const dtReal dt, dtCrowdAgentDebugInfo* debug);
 
 	/// [UE] Split update into several smaller components: integrate velocities and handle collisions
 	/// @param[in]		dt		Delta time in seconds
 	/// @param[in]		nagents	Number of active agents
-	void updateStepMove(const dtReal dt, dtCrowdAgentDebugInfo* debug);
+	NAVMESH_API void updateStepMove(const dtReal dt, dtCrowdAgentDebugInfo* debug);
 
 	/// [UE] Split update into several smaller components: corridor updates at new position
 	/// @param[in]		dt		Delta time in seconds
 	/// @param[in]		nagents	Number of active agents
-	void updateStepCorridor(const dtReal dt, dtCrowdAgentDebugInfo* debug);
+	NAVMESH_API void updateStepCorridor(const dtReal dt, dtCrowdAgentDebugInfo* debug);
 
 	/// [UE] Split update into several smaller components: offmesh anims
 	/// @param[in]		dt		Delta time in seconds
 	/// @param[in]		nagents	Number of active agents
-	void updateStepOffMeshAnim(const dtReal dt, dtCrowdAgentDebugInfo* debug);
+	NAVMESH_API void updateStepOffMeshAnim(const dtReal dt, dtCrowdAgentDebugInfo* debug);
 
 	/// [UE] Split update into several smaller components: offmesh link velocity (instead of playing animation)
 	/// @param[in]		dt		Delta time in seconds
 	/// @param[in]		nagents	Number of active agents
-	void updateStepOffMeshVelocity(const dtReal dt, dtCrowdAgentDebugInfo* debug);
+	NAVMESH_API void updateStepOffMeshVelocity(const dtReal dt, dtCrowdAgentDebugInfo* debug);
 
 	/// [UE] Set time between attempts to restore agents state
 	/// @param[in]		t		Time in seconds
-	void setAgentCheckInterval(const dtReal t);
+	NAVMESH_API void setAgentCheckInterval(const dtReal t);
 
 	/// [UE] Set agent corridor, works only just after requesting move target
 	/// when agent didn't start any pathfinding operations yet
@@ -466,40 +466,40 @@ public:
 	/// @param[in]		idx		The agent index. [Limits: 0 <= value < #getAgentCount()]
 	/// @param[in]		path	The path corridor. [(polyRef) * @p npolys]
 	/// @param[in]		npath	The number of polygons in the path.
-	bool setAgentCorridor(const int idx, const dtPolyRef* path, const int npath);
+	NAVMESH_API bool setAgentCorridor(const int idx, const dtPolyRef* path, const int npath);
 
 	/// [UE] Set visibility optimization to use single area raycasts
 	/// This will prevent from cutting through polys marked as different area
 	/// which could have been avoided in corridor's path
 	/// @param[in]		bEnable	New state of single area raycast mode
-	void setSingleAreaVisibilityOptimization(bool bEnable);
+	NAVMESH_API void setSingleAreaVisibilityOptimization(bool bEnable);
 
 	/// [UE] Set offmesh connection pruning
 	/// This will allow removing offmesh connection poly ref from corridor
 	/// as soon as offmesh connection anim is triggered (default behavior)
-	void setPruneStartedOffmeshConnections(bool bRemoveFromCorridor);
+	NAVMESH_API void setPruneStartedOffmeshConnections(bool bRemoveFromCorridor);
 
 	/// [UE]
-	void setEarlyReachTestOptimization(bool bEnable);
+	NAVMESH_API void setEarlyReachTestOptimization(bool bEnable);
 
 	/// [UE] Set agent radius multiplier for offseting path from corners
-	void setPathOffsetRadiusMultiplier(dtReal RadiusMultiplier);
+	NAVMESH_API void setPathOffsetRadiusMultiplier(dtReal RadiusMultiplier);
 
 	/// [UE] Set separation filter param
-	void setSeparationFilter(dtReal InFilter);
+	NAVMESH_API void setSeparationFilter(dtReal InFilter);
 
 	/// [UE] Check if agent moved away from its path corridor
-	bool isOutsideCorridor(const int idx) const;
+	NAVMESH_API bool isOutsideCorridor(const int idx) const;
 
 	/// Gets the filter used by the crowd.
 	///  @param[in]		idx		[UE] The agent index. [Limits: 0 <= value < #getAgentCount()]
 	/// @return The filter used by the crowd.
-	const dtQueryFilter* getFilter(const int idx) const;
+	NAVMESH_API const dtQueryFilter* getFilter(const int idx) const;
 
 	/// Gets the filter used by the crowd.
 	///  @param[in]		idx		[UE] The agent index. [Limits: 0 <= value < #getAgentCount()]
 	/// @return The filter used by the crowd.
-	dtQueryFilter* getEditableFilter(const int idx);
+	NAVMESH_API dtQueryFilter* getEditableFilter(const int idx);
 
 	/// Gets the search extents [(x, y, z)] used by the crowd for query operations. 
 	/// @return The search extents used by the crowd. [(x, y, z)]
diff --git a/Engine/Source/Runtime/Navmesh/Public/DetourCrowd/DetourLocalBoundary.h b/Engine/Source/Runtime/Navmesh/Public/DetourCrowd/DetourLocalBoundary.h
--- a/Engine/Source/Runtime/Navmesh/Public/DetourCrowd/DetourLocalBoundary.h
+++ b/Engine/Source/Runtime/Navmesh/Public/DetourCrowd/DetourLocalBoundary.h
@@ -30,7 +30,7 @@ class dtNavMeshQuery;
 class dtQueryFilter;
 class dtSharedBoundary;
 
-class NAVMESH_API dtLocalBoundary
+class dtLocalBoundary
 {
 	static const int MAX_LOCAL_SEGS = 8;
 	static const int MAX_LOCAL_POLYS = 16;
@@ -48,28 +48,28 @@ class NAVMESH_API dtLocalBoundary
 	int m_nsegs;
 	int m_npolys;
 
-	void addSegment(const dtReal dist, const dtReal* seg, int flags = 0);
+	NAVMESH_API void addSegment(const dtReal dist, const dtReal* seg, int flags = 0);
 	
 public:
-	dtLocalBoundary();
-	~dtLocalBoundary();
+	NAVMESH_API dtLocalBoundary();
+	NAVMESH_API ~dtLocalBoundary();
 	
-	void reset();
+	NAVMESH_API void reset();
 
 	// [UE: new sections: link removal, path corridor, direction]
-	void update(dtPolyRef ref, const dtReal* pos, const dtReal collisionQueryRange,
+	NAVMESH_API void update(dtPolyRef ref, const dtReal* pos, const dtReal collisionQueryRange,
 		const bool bIgnoreAtEnd, const dtReal* endPos,
 		const dtPolyRef* path, const int npath,
 		const dtReal* moveDir,
 		dtNavMeshQuery* navquery, const dtQueryFilter* filter);
 
-	void update(const dtSharedBoundary* sharedData, const int sharedIdx,
+	NAVMESH_API void update(const dtSharedBoundary* sharedData, const int sharedIdx,
 		const dtReal* pos, const dtReal collisionQueryRange,
 		const bool bIgnoreAtEnd, const dtReal* endPos,
 		const dtPolyRef* path, const int npath, const dtReal* moveDir,
 		dtNavMeshQuery* navquery, const dtQueryFilter* filter);
 	
-	bool isValid(dtNavMeshQuery* navquery, const dtQueryFilter* filter);
+	NAVMESH_API bool isValid(dtNavMeshQuery* navquery, const dtQueryFilter* filter);
 	
 	inline const dtReal* getCenter() const { return m_center; }
 	inline int getSegmentCount() const { return m_nsegs; }
diff --git a/Engine/Source/Runtime/Navmesh/Public/DetourCrowd/DetourObstacleAvoidance.h b/Engine/Source/Runtime/Navmesh/Public/DetourCrowd/DetourObstacleAvoidance.h
--- a/Engine/Source/Runtime/Navmesh/Public/DetourCrowd/DetourObstacleAvoidance.h
+++ b/Engine/Source/Runtime/Navmesh/Public/DetourCrowd/DetourObstacleAvoidance.h
@@ -41,18 +41,18 @@ struct dtObstacleSegment
 	unsigned char canIgnore : 1;
 };
 
-class NAVMESH_API dtObstacleAvoidanceDebugData
+class dtObstacleAvoidanceDebugData
 {
 public:
-	dtObstacleAvoidanceDebugData();
-	~dtObstacleAvoidanceDebugData();
+	NAVMESH_API dtObstacleAvoidanceDebugData();
+	NAVMESH_API ~dtObstacleAvoidanceDebugData();
 	
-	bool init(const int maxSamples);
-	void reset();
-	void addSample(const dtReal* vel, const dtReal ssize, const dtReal pen,
+	NAVMESH_API bool init(const int maxSamples);
+	NAVMESH_API void reset();
+	NAVMESH_API void addSample(const dtReal* vel, const dtReal ssize, const dtReal pen,
 				   const dtReal vpen, const dtReal vcpen, const dtReal spen, const dtReal tpen);
 	
-	void normalizeSamples();
+	NAVMESH_API void normalizeSamples();
 	
 	inline int getSampleCount() const { return m_nsamples; }
 	inline const dtReal* getSampleVelocity(const int i) const { return &m_vel[i*3]; }
@@ -83,7 +83,7 @@ static const int DT_MAX_PATTERN_DIVS = 32;		///< Max numver of adaptive divs.
 static const int DT_MAX_PATTERN_RINGS = 4;		///< Max number of adaptive rings.
 static const int DT_MAX_CUSTOM_SAMPLES = 16;	///< Max number of custom samples in single pattern
 
-struct NAVMESH_API dtObstacleAvoidanceParams
+struct dtObstacleAvoidanceParams
 {
 	dtReal velBias;
 	dtReal weightDesVel;
@@ -105,35 +105,35 @@ struct dtObstacleAvoidancePattern
 	int nsamples;							///< Number of samples
 };
 
-class NAVMESH_API dtObstacleAvoidanceQuery
+class dtObstacleAvoidanceQuery
 {
 public:
-	dtObstacleAvoidanceQuery();
-	~dtObstacleAvoidanceQuery();
+	NAVMESH_API dtObstacleAvoidanceQuery();
+	NAVMESH_API ~dtObstacleAvoidanceQuery();
 	
-	bool init(const int maxCircles, const int maxSegments, const int maxCustomPatterns);
+	NAVMESH_API bool init(const int maxCircles, const int maxSegments, const int maxCustomPatterns);
 	
-	void reset();
+	NAVMESH_API void reset();
 
-	void addCircle(const dtReal* pos, const dtReal rad,
+	NAVMESH_API void addCircle(const dtReal* pos, const dtReal rad,
 				   const dtReal* vel, const dtReal* dvel);
 				   
-	void addSegment(const dtReal* p, const dtReal* q, int flags = 0);
+	NAVMESH_API void addSegment(const dtReal* p, const dtReal* q, int flags = 0);
 
 	// [UE] store new sampling pattern
-	bool setCustomSamplingPattern(int idx, const dtReal* angles, const dtReal* radii, int nsamples);
+	NAVMESH_API bool setCustomSamplingPattern(int idx, const dtReal* angles, const dtReal* radii, int nsamples);
 
 	// [UE] get custom sampling pattern
-	bool getCustomSamplingPattern(int idx, dtReal* angles, dtReal* radii, int* nsamples);
+	NAVMESH_API bool getCustomSamplingPattern(int idx, dtReal* angles, dtReal* radii, int* nsamples);
 
 	// [UE] sample velocity using custom patterns
-	int sampleVelocityCustom(const dtReal* pos, const dtReal rad,
+	NAVMESH_API int sampleVelocityCustom(const dtReal* pos, const dtReal rad,
 					 		 const dtReal vmax, const dtReal vmult,
 							 const dtReal* vel, const dtReal* dvel, dtReal* nvel,
 							 const dtObstacleAvoidanceParams* params,
 							 dtObstacleAvoidanceDebugData* debug = 0);
 
-	int sampleVelocityAdaptive(const dtReal* pos, const dtReal rad,
+	NAVMESH_API int sampleVelocityAdaptive(const dtReal* pos, const dtReal rad,
 							   const dtReal vmax, const dtReal vmult,
 							   const dtReal* vel, const dtReal* dvel, dtReal* nvel,
 							   const dtObstacleAvoidanceParams* params, 
@@ -162,15 +162,15 @@ public:
 
 private:
 
-	void prepare(const dtReal* pos, const dtReal* dvel);
+	NAVMESH_API void prepare(const dtReal* pos, const dtReal* dvel);
 
-	dtReal processSample(const dtReal* vcand, const dtReal cs,
+	NAVMESH_API dtReal processSample(const dtReal* vcand, const dtReal cs,
 						const dtReal* pos, const dtReal rad,
 						const dtReal* vel, const dtReal* dvel,
 						dtObstacleAvoidanceDebugData* debug);
 
-	dtObstacleCircle* insertCircle(const dtReal dist);
-	dtObstacleSegment* insertSegment(const dtReal dist);
+	NAVMESH_API dtObstacleCircle* insertCircle(const dtReal dist);
+	NAVMESH_API dtObstacleSegment* insertSegment(const dtReal dist);
 
 	dtObstacleAvoidanceParams m_params;
 	dtReal m_invHorizTime;
diff --git a/Engine/Source/Runtime/Navmesh/Public/DetourCrowd/DetourPathCorridor.h b/Engine/Source/Runtime/Navmesh/Public/DetourCrowd/DetourPathCorridor.h
--- a/Engine/Source/Runtime/Navmesh/Public/DetourCrowd/DetourPathCorridor.h
+++ b/Engine/Source/Runtime/Navmesh/Public/DetourCrowd/DetourPathCorridor.h
@@ -32,7 +32,7 @@ class dtQueryFilter;
 
 /// Represents a dynamic polygon corridor used to plan agent movement.
 /// @ingroup crowd, detour
-class NAVMESH_API dtPathCorridor
+class dtPathCorridor
 {
 	dtReal m_pos[3];
 	dtReal m_target[3];
@@ -50,18 +50,18 @@ class NAVMESH_API dtPathCorridor
 	int m_maxPath;
 	
 public:
-	dtPathCorridor();
-	~dtPathCorridor();
+	NAVMESH_API dtPathCorridor();
+	NAVMESH_API ~dtPathCorridor();
 	
 	/// Allocates the corridor's path buffer. 
 	///  @param[in]		maxPath		The maximum path size the corridor can handle.
 	/// @return True if the initialization succeeded.
-	bool init(const int maxPath);
+	NAVMESH_API bool init(const int maxPath);
 	
 	/// Resets the path corridor to the specified position.
 	///  @param[in]		ref		The polygon reference containing the position.
 	///  @param[in]		pos		The new position in the corridor. [(x, y, z)]
-	void reset(dtPolyRef ref, const dtReal* pos);
+	NAVMESH_API void reset(dtPolyRef ref, const dtReal* pos);
 	
 	/// Finds the corners in the corridor from the position toward the target. (The straightened path.)
 	///  @param[out]	cornerVerts		The corner vertices. [(x, y, z) * cornerCount] [Size: <= maxCorners]
@@ -75,7 +75,7 @@ public:
 	///  @param[in]		earlyReachDistance	[UE] Radius for early reach detection
 	///  @param[in]		bAllowEarlyReach [UE] Check if corner skipping for EarlyReachTest is available now
 	/// @return The number of corners returned in the corner buffers. [0 <= value <= @p maxCorners]
-	int findCorners(dtReal* cornerVerts, unsigned char* cornerFlags,
+	NAVMESH_API int findCorners(dtReal* cornerVerts, unsigned char* cornerFlags,
 					dtPolyRef* cornerPolys, const int maxCorners,
 					dtNavMeshQuery* navquery, const dtQueryFilter* filter,
 					dtReal pathOffsetDistance, dtReal earlyReachDistance, bool bAllowEarlyReach = true);
@@ -85,57 +85,57 @@ public:
 	///  @param[in]		pathOptimizationRange	The maximum range to search. [Limit: > 0]
 	///  @param[in]		navquery				The query object used to build the corridor.
 	///  @param[in]		filter					The filter to apply to the operation.			
-	bool optimizePathVisibility(const dtReal* next, const dtReal pathOptimizationRange,
+	NAVMESH_API bool optimizePathVisibility(const dtReal* next, const dtReal pathOptimizationRange,
 								dtNavMeshQuery* navquery, const dtQueryFilter* filter);
 	
 	/// Attempts to optimize the path using a local area search. (Partial replanning.) 
 	///  @param[in]		navquery	The query object used to build the corridor.
 	///  @param[in]		filter		The filter to apply to the operation.	
-	bool optimizePathTopology(dtNavMeshQuery* navquery, const dtQueryFilter* filter);
+	NAVMESH_API bool optimizePathTopology(dtNavMeshQuery* navquery, const dtQueryFilter* filter);
 	
-	bool moveOverOffmeshConnection(dtPolyRef offMeshConRef, dtPolyRef* refs,
+	NAVMESH_API bool moveOverOffmeshConnection(dtPolyRef offMeshConRef, dtPolyRef* refs,
 								   const dtReal* agentPos,
 								   dtReal* startPos, dtReal* endPos,
 								   dtNavMeshQuery* navquery);
 
 	/// [UE] check if offmesh connection can be traversed, but don't modify corridor yet
-	bool canMoveOverOffmeshConnection(dtPolyRef offMeshConRef, dtPolyRef* refs,
+	NAVMESH_API bool canMoveOverOffmeshConnection(dtPolyRef offMeshConRef, dtPolyRef* refs,
 		const dtReal* agentPos, dtReal* startPos, dtReal* endPos,
 		dtNavMeshQuery* navquery) const;
 
 	/// [UE] remove offmesh connection from corridor
-	void pruneOffmeshConenction(dtPolyRef offMeshConRef);
+	NAVMESH_API void pruneOffmeshConenction(dtPolyRef offMeshConRef);
 
-	bool fixPathStart(dtPolyRef safeRef, const dtReal* safePos);
+	NAVMESH_API bool fixPathStart(dtPolyRef safeRef, const dtReal* safePos);
 
-	bool trimInvalidPath(dtPolyRef safeRef, const dtReal* safePos,
+	NAVMESH_API bool trimInvalidPath(dtPolyRef safeRef, const dtReal* safePos,
 						 dtNavMeshQuery* navquery, const dtQueryFilter* filter);
 	
 	/// Checks the current corridor path to see if its polygon references remain valid. 
 	///  @param[in]		maxLookAhead	The number of polygons from the beginning of the corridor to search.
 	///  @param[in]		navquery		The query object used to build the corridor.
 	///  @param[in]		filter			The filter to apply to the operation.	
-	bool isValid(const int maxLookAhead, dtNavMeshQuery* navquery, const dtQueryFilter* filter);
+	NAVMESH_API bool isValid(const int maxLookAhead, dtNavMeshQuery* navquery, const dtQueryFilter* filter);
 	
 	/// Moves the position from the current location to the desired location, adjusting the corridor 
 	/// as needed to reflect the change.
 	///  @param[in]		npos		The desired new position. [(x, y, z)]
 	///  @param[in]		navquery	The query object used to build the corridor.
 	///  @param[in]		filter		The filter to apply to the operation.
-	bool movePosition(const dtReal* npos, dtNavMeshQuery* navquery, const dtQueryFilter* filter);
+	NAVMESH_API bool movePosition(const dtReal* npos, dtNavMeshQuery* navquery, const dtQueryFilter* filter);
 
 	/// Moves the target from the curent location to the desired location, adjusting the corridor
 	/// as needed to reflect the change. 
 	///  @param[in]		npos		The desired new target position. [(x, y, z)]
 	///  @param[in]		navquery	The query object used to build the corridor.
 	///  @param[in]		filter		The filter to apply to the operation.
-	void moveTargetPosition(const dtReal* npos, dtNavMeshQuery* navquery, const dtQueryFilter* filter);
+	NAVMESH_API void moveTargetPosition(const dtReal* npos, dtNavMeshQuery* navquery, const dtQueryFilter* filter);
 	
 	/// Loads a new path and target into the corridor.
 	///  @param[in]		target		The target location within the last polygon of the path. [(x, y, z)]
 	///  @param[in]		path		The path corridor. [(polyRef) * @p npolys]
 	///  @param[in]		npath		The number of polygons in the path.
-	void setCorridor(const dtReal* target, const dtPolyRef* polys, const int npath);
+	NAVMESH_API void setCorridor(const dtReal* target, const dtPolyRef* polys, const int npath);
 	
 	/// Gets the current position within the corridor. (In the first polygon.)
 	/// @return The current position within the corridor.
diff --git a/Engine/Source/Runtime/Navmesh/Public/DetourTileCache/DetourTileCache.h b/Engine/Source/Runtime/Navmesh/Public/DetourTileCache/DetourTileCache.h
--- a/Engine/Source/Runtime/Navmesh/Public/DetourTileCache/DetourTileCache.h
+++ b/Engine/Source/Runtime/Navmesh/Public/DetourTileCache/DetourTileCache.h
@@ -94,7 +94,7 @@ struct dtTileCacheParams
 //@UE END
 };
 
-struct NAVMESH_API dtTileCacheMeshProcess
+struct dtTileCacheMeshProcess
 {
 	virtual void markAreas(struct dtTileCacheLayer* layer, const dtReal* orig, const dtReal cs, const dtReal ch) = 0;
 	
@@ -110,11 +110,11 @@ struct NAVMESH_API dtTileCacheMeshProcess
 };
 
 
-class NAVMESH_API dtTileCache
+class dtTileCache
 {
 public:
-	dtTileCache();
-	~dtTileCache();
+	NAVMESH_API dtTileCache();
+	NAVMESH_API ~dtTileCache();
 	
 	struct dtTileCacheAlloc* getAlloc() { return m_talloc; }
 	struct dtTileCacheCompressor* getCompressor() { return m_tcomp; }
@@ -127,40 +127,40 @@ public:
 	inline int getObstacleCount() const { return m_params.maxObstacles; }
 	inline const dtTileCacheObstacle* getObstacle(const int i) const { return &m_obstacles[i]; }
 	
-	const dtTileCacheObstacle* getObstacleByRef(dtObstacleRef ref);
+	NAVMESH_API const dtTileCacheObstacle* getObstacleByRef(dtObstacleRef ref);
 	
-	dtObstacleRef getObstacleRef(const dtTileCacheObstacle* obmin) const;
+	NAVMESH_API dtObstacleRef getObstacleRef(const dtTileCacheObstacle* obmin) const;
 	
-	dtStatus init(const dtTileCacheParams* params,
+	NAVMESH_API dtStatus init(const dtTileCacheParams* params,
 				  struct dtTileCacheAlloc* talloc,
 				  struct dtTileCacheCompressor* tcomp,
 				  struct dtTileCacheMeshProcess* tmproc);
 	
-	int getTilesAt(const int tx, const int ty, dtCompressedTileRef* tiles, const int maxTiles) const ;
+	NAVMESH_API int getTilesAt(const int tx, const int ty, dtCompressedTileRef* tiles, const int maxTiles) const ;
 	
-	dtCompressedTile* getTileAt(const int tx, const int ty, const int tlayer);
-	dtCompressedTileRef getTileRef(const dtCompressedTile* tile) const;
-	const dtCompressedTile* getTileByRef(dtCompressedTileRef ref) const;
+	NAVMESH_API dtCompressedTile* getTileAt(const int tx, const int ty, const int tlayer);
+	NAVMESH_API dtCompressedTileRef getTileRef(const dtCompressedTile* tile) const;
+	NAVMESH_API const dtCompressedTile* getTileByRef(dtCompressedTileRef ref) const;
 	
-	dtStatus addTile(unsigned char* data, const int dataSize, unsigned char flags, dtCompressedTileRef* result);
+	NAVMESH_API dtStatus addTile(unsigned char* data, const int dataSize, unsigned char flags, dtCompressedTileRef* result);
 	
-	dtStatus removeTile(dtCompressedTileRef ref, unsigned char** data, int* dataSize);
+	NAVMESH_API dtStatus removeTile(dtCompressedTileRef ref, unsigned char** data, int* dataSize);
 	
-	dtStatus addObstacle(const dtReal* pos, const dtReal radius, const dtReal height, dtObstacleRef* result);
-	dtStatus removeObstacle(const dtObstacleRef ref);
+	NAVMESH_API dtStatus addObstacle(const dtReal* pos, const dtReal radius, const dtReal height, dtObstacleRef* result);
+	NAVMESH_API dtStatus removeObstacle(const dtObstacleRef ref);
 	
-	dtStatus queryTiles(const dtReal* bmin, const dtReal* bmax,
+	NAVMESH_API dtStatus queryTiles(const dtReal* bmin, const dtReal* bmax,
 						dtCompressedTileRef* results, int* resultCount, const int maxResults) const;
 	
-	dtStatus update(const dtReal /*dt*/, class dtNavMesh* navmesh);
+	NAVMESH_API dtStatus update(const dtReal /*dt*/, class dtNavMesh* navmesh);
 	
-	dtStatus buildNavMeshTilesAt(const int tx, const int ty, class dtNavMesh* navmesh);
+	NAVMESH_API dtStatus buildNavMeshTilesAt(const int tx, const int ty, class dtNavMesh* navmesh);
 	
-	dtStatus buildNavMeshTile(const dtCompressedTileRef ref, class dtNavMesh* navmesh);
+	NAVMESH_API dtStatus buildNavMeshTile(const dtCompressedTileRef ref, class dtNavMesh* navmesh);
 	
-	void calcTightTileBounds(const struct dtTileCacheLayerHeader* header, dtReal* bmin, dtReal* bmax) const;
+	NAVMESH_API void calcTightTileBounds(const struct dtTileCacheLayerHeader* header, dtReal* bmin, dtReal* bmax) const;
 	
-	void getObstacleBounds(const struct dtTileCacheObstacle* ob, dtReal* bmin, dtReal* bmax) const;
+	NAVMESH_API void getObstacleBounds(const struct dtTileCacheObstacle* ob, dtReal* bmin, dtReal* bmax) const;
 	
 
 	/// Encodes a tile id.
diff --git a/Engine/Source/Runtime/Navmesh/Public/DetourTileCache/DetourTileCacheBuilder.h b/Engine/Source/Runtime/Navmesh/Public/DetourTileCache/DetourTileCacheBuilder.h
--- a/Engine/Source/Runtime/Navmesh/Public/DetourTileCache/DetourTileCacheBuilder.h
+++ b/Engine/Source/Runtime/Navmesh/Public/DetourTileCache/DetourTileCacheBuilder.h
@@ -110,13 +110,13 @@ struct dtTileCacheDistanceField
 	unsigned short* data;	///< distance for every cell in layer
 };
 
-class NAVMESH_API dtTileCacheLogContext
+class dtTileCacheLogContext
 {
 public:
 	/// Logs a message.
 	///  @param[in]		category	The category of the message.
 	///  @param[in]		format		The message.
-	void dtLog(const char* format, ...);
+	NAVMESH_API void dtLog(const char* format, ...);
 
 protected:
 
@@ -129,7 +129,7 @@ protected:
 
 //@UE END
 
-struct NAVMESH_API dtTileCacheAlloc
+struct dtTileCacheAlloc
 {
 	virtual ~dtTileCacheAlloc() = default;
 
@@ -148,7 +148,7 @@ struct NAVMESH_API dtTileCacheAlloc
 	}
 };
 
-struct NAVMESH_API dtTileCacheCompressor
+struct dtTileCacheCompressor
 {
 	virtual int maxCompressedSize(const int bufferSize) = 0;
 	virtual dtStatus compress(const unsigned char* buffer, const int bufferSize,
diff --git a/Engine/Source/Runtime/Navmesh/Public/Recast/Recast.h b/Engine/Source/Runtime/Navmesh/Public/Recast/Recast.h
--- a/Engine/Source/Runtime/Navmesh/Public/Recast/Recast.h
+++ b/Engine/Source/Runtime/Navmesh/Public/Recast/Recast.h
@@ -161,7 +161,7 @@ NAVMESH_API DECLARE_LOG_CATEGORY_EXTERN(LogRecast, Log, All);
 /// Provides an interface for optional logging and performance tracking of the Recast 
 /// build process.
 /// @ingroup recast
-class NAVMESH_API rcContext
+class rcContext
 {
 public:
 
@@ -180,7 +180,7 @@ public:
 	/// Logs a message.
 	///  @param[in]		category	The category of the message.
 	///  @param[in]		format		The message.
-	void log(const rcLogCategory category, const char* format, ...);
+	NAVMESH_API void log(const rcLogCategory category, const char* format, ...);
 
 	/// Enables or disables the performance timers.
 	///  @param[in]		state	TRUE if timers should be enabled.
diff --git a/Engine/Source/Runtime/Net/Core/Public/Net/Core/Analytics/NetAnalytics.h b/Engine/Source/Runtime/Net/Core/Public/Net/Core/Analytics/NetAnalytics.h
--- a/Engine/Source/Runtime/Net/Core/Public/Net/Core/Analytics/NetAnalytics.h
+++ b/Engine/Source/Runtime/Net/Core/Public/Net/Core/Analytics/NetAnalytics.h
@@ -105,7 +105,7 @@ struct FLocalNetAnalyticsStruct : public FVirtualDestructor
 /**
  * Subclassed struct which holds net analytics data which is aggregated or to-be-aggregated, and dispatched upon SendAnalytics
  */
-struct NETCORE_API FNetAnalyticsData : public FVirtualDestructor, public TSharedFromThis<FNetAnalyticsData, NetAnalyticsThreadSafety>
+struct FNetAnalyticsData : public FVirtualDestructor, public TSharedFromThis<FNetAnalyticsData, NetAnalyticsThreadSafety>
 {
 	friend class FNetAnalyticsAggregator;
 	friend struct FNetAnalyticsDataDeleter;
@@ -134,7 +134,7 @@ protected:
 	/**
 	 * Called by the Net Analytics Aggregator, when ready to send analytics - some implementations don't send immediately
 	 */
-	virtual void InternalSendAnalytics();
+	NETCORE_API virtual void InternalSendAnalytics();
 
 	/**
 	 * Called just before the last shared reference to this data is released - used to trigger analytics send in threadsafe version
@@ -176,18 +176,18 @@ template<class TDataStruct> struct TBasicNetAnalyticsData : public FNetAnalytics
  *
  * The last thread to release this analytics data, triggers SendAnalytics - thread safety is partially provided by shared pointer atomics.
  */
-struct NETCORE_API FThreadedNetAnalyticsData : public FNetAnalyticsData
+struct FThreadedNetAnalyticsData : public FNetAnalyticsData
 {
 	/**
 	 * Default constructor
 	 */
-	FThreadedNetAnalyticsData();
+	NETCORE_API FThreadedNetAnalyticsData();
 
 
 protected:
-	virtual void InternalSendAnalytics() override;
+	NETCORE_API virtual void InternalSendAnalytics() override;
 
-	virtual void NotifyFinalRelease() override;
+	NETCORE_API virtual void NotifyFinalRelease() override;
 
 
 protected:
@@ -321,31 +321,31 @@ struct FNetAnalyticsDataDeleter
 /**
  * Central object (usually within NetDriver) which handles registration/retrieval/type-checking of net analytics data holders.
  */
-class NETCORE_API FNetAnalyticsAggregator
+class FNetAnalyticsAggregator
 {
 public:
 	/**
 	 * Base constructor
 	 */
-	FNetAnalyticsAggregator(TSharedPtr<IAnalyticsProvider> InProvider, FName InNetDriverName);
+	NETCORE_API FNetAnalyticsAggregator(TSharedPtr<IAnalyticsProvider> InProvider, FName InNetDriverName);
 
 	FNetAnalyticsAggregator() = delete;
 
 	/**
 	 * Initialize the net analytics aggregator
 	 */
-	void Init();
+	NETCORE_API void Init();
 
 	/**
 	 * Initialize the net analytics aggregator config - must support hotfixing
 	 */
-	void InitConfig();
+	NETCORE_API void InitConfig();
 
 	/**
 	 * Tells the analytics data holders to finish aggregating their analytics data, and to dispatch it.
 	 * Only called once, at NetDriver shutdown.
 	 */
-	void SendAnalytics();
+	NETCORE_API void SendAnalytics();
 
 
 	/**
@@ -359,7 +359,7 @@ public:
 	 * @param InTypeName	Compile-time derived type name for the analytics data - for type checking
 	 * @return				Returns the newly added analytics data shared reference
 	 */
-	TNetAnalyticsDataPtr<> RegisterAnalyticsData_Internal(TNetAnalyticsDataRef<> InData, const FName& InDataName, FString InTypeName);
+	NETCORE_API TNetAnalyticsDataPtr<> RegisterAnalyticsData_Internal(TNetAnalyticsDataRef<> InData, const FName& InDataName, FString InTypeName);
 
 
 	/** Accessor for AnalyticsProvider */
diff --git a/Engine/Source/Runtime/Net/Core/Public/Net/Core/Analytics/NetStatsUtils.h b/Engine/Source/Runtime/Net/Core/Public/Net/Core/Analytics/NetStatsUtils.h
--- a/Engine/Source/Runtime/Net/Core/Public/Net/Core/Analytics/NetStatsUtils.h
+++ b/Engine/Source/Runtime/Net/Core/Public/Net/Core/Analytics/NetStatsUtils.h
@@ -112,7 +112,7 @@ enum class EBinnedValueMode : uint8
  * Base class for moving values, implementing non-templatized data-structure/methods.
  * See TBinnedMovingValue.
  */
-class NETCORE_API FBinnedMovingValueBase
+class FBinnedMovingValueBase
 {
 protected:
 	struct FBin
@@ -126,7 +126,7 @@ protected:
 
 protected:
 	FBinnedMovingValueBase() = delete;
-	explicit FBinnedMovingValueBase(const TArrayView<FBin> InBins, double InTimePerBin, EBinnedValueMode InMode);
+	NETCORE_API explicit FBinnedMovingValueBase(const TArrayView<FBin> InBins, double InTimePerBin, EBinnedValueMode InMode);
 
 	/**
 	 * Internal implementation for AddMeasurement, which outputs a sample if the bins have filled and rolled back over to FirstBinIndex.
@@ -135,7 +135,7 @@ protected:
 	 * @param Value			The measurement value
 	 * @param OutSample		Outputs a sample, if all bins have filled and rolled back over to FirstBinIndex
 	 */
-	void AddMeasurement_Implementation(double TimeVal, double Value, TOptional<double>& OutSample);
+	NETCORE_API void AddMeasurement_Implementation(double TimeVal, double Value, TOptional<double>& OutSample);
 
 public:
 	/**
@@ -143,15 +143,15 @@ public:
 	 *
 	 * @return	Outputs the calculated sample value
 	 */
-	double GetSample() const;
+	NETCORE_API double GetSample() const;
 
 	/**
 	 * Resets the accumulated bin/sample data.
 	 */
-	void Reset();
+	NETCORE_API void Reset();
 
 private:
-	void ResetNewAndSkippedBins(int32 BinIdx, TOptional<double>& OutSample);
+	NETCORE_API void ResetNewAndSkippedBins(int32 BinIdx, TOptional<double>& OutSample);
 
 
 private:
@@ -367,18 +367,18 @@ enum EMinMaxValueMode
  * This base class can handle reading.
  * The TSampleMinMaxAvg subclass is compatible as a consumer for TSampleProducer, and is used for writing.
  */
-class NETCORE_API FSampleMinMaxAvg
+class FSampleMinMaxAvg
 {
 	template<EMinMaxValueMode, int32> friend class TSampleMinMaxAvg;
 
 private:
-	void AddSample_Internal(double Value);
+	NETCORE_API void AddSample_Internal(double Value);
 
 public:
 	/**
 	 * Resets the accumulated sample data
 	 */
-	void Reset();
+	NETCORE_API void Reset();
 
 	/**
 	 * Gets the minimum sample value encountered
diff --git a/Engine/Source/Runtime/Net/Core/Public/Net/Core/Connection/EscalationStates.h b/Engine/Source/Runtime/Net/Core/Public/Net/Core/Connection/EscalationStates.h
--- a/Engine/Source/Runtime/Net/Core/Public/Net/Core/Connection/EscalationStates.h
+++ b/Engine/Source/Runtime/Net/Core/Public/Net/Core/Connection/EscalationStates.h
@@ -155,7 +155,7 @@ struct FHasHitAnyQuotaParms
  * NOTE: Subclasses must also implement the TEscalationStateStatics interface.
  */
 USTRUCT()
-struct NETCORE_API FEscalationState : public FStateStruct
+struct FEscalationState : public FStateStruct
 {
 	friend UE::Net::FEscalationManager;
 	friend UEscalationManagerConfig;
@@ -205,25 +205,25 @@ public:
 	/**
 	 * Whether or not this escalation state is a 'dormant' state, which allows the escalation manager to disable ticking
 	 */
-	bool IsDormant() const;
+	NETCORE_API bool IsDormant() const;
 
 	/**
 	 * Gets the highest counter time period specified by escalation state settings (used for limiting 'CountersPerPeriodHistory' size)
 	 *
 	 * @return		The highest counter time period in the state config settings
 	 */
-	int8 GetHighestTimePeriod() const;
+	NETCORE_API int8 GetHighestTimePeriod() const;
 
 	/**
 	 * Gets all counter time periods specified by the escalation state settings
 	 *
 	 * @return		All of the time periods in the state config settings
 	 */
-	const TArray<int8>& GetAllTimePeriods() const;
+	NETCORE_API const TArray<int8>& GetAllTimePeriods() const;
 
 
 protected:
-	virtual void ValidateConfigInternal() override;
+	NETCORE_API virtual void ValidateConfigInternal() override;
 
 	/**
 	 * Validate that the specified escalation time period is within thresholds, and clamp it if necessary
@@ -232,7 +232,7 @@ protected:
 	 * @param PropertyName	The name of the property, for logging
 	 * @param Requirement	Whether the time period value must be set, or is optional
 	 */
-	void ValidateTimePeriod(int8& Value, const TCHAR* PropertyName, EValidateTime Requirement=EValidateTime::Optional);
+	NETCORE_API void ValidateTimePeriod(int8& Value, const TCHAR* PropertyName, EValidateTime Requirement=EValidateTime::Optional);
 
 
 private:
@@ -376,7 +376,7 @@ enum class EEscalateResult : uint8
 /**
  * Manages initialization/application of escalation states, and tracking various counters/timers used to calculate state change quotas/thresholds.
  */
-class NETCORE_API FEscalationManager
+class FEscalationManager
 {
 	template<typename, typename, typename> friend class TEscalationManager;
 	friend UEscalationManagerConfig;
@@ -445,9 +445,9 @@ private:
 
 
 private:
-	FEscalationManager(FEscalationManagerParms Parms);
+	NETCORE_API FEscalationManager(FEscalationManagerParms Parms);
 
-	void InitParms(FEscalationManagerInitParms Parms);
+	NETCORE_API void InitParms(FEscalationManagerInitParms Parms);
 
 	FEscalationManager() = delete;
 
@@ -456,7 +456,7 @@ private:
 	 *
 	 * @param ConfigParms	Specifies the parameters for the config section/object
 	 */
-	void InitConfig(FStateConfigParms ConfigParms);
+	NETCORE_API void InitConfig(FStateConfigParms ConfigParms);
 
 	/**
 	 * Updates the current escalation manager severity state
@@ -466,7 +466,7 @@ private:
 	 * @param ReasonContext		Additional customizable context for the escalation change
 	 * @return					The result of the severity update. Escalation/Deescalation, or NoChange
 	 */
-	EEscalateResult UpdateSeverity(ESeverityUpdate Update, EEscalateReason Reason, FString ReasonContext=TEXT(""));
+	NETCORE_API EEscalateResult UpdateSeverity(ESeverityUpdate Update, EEscalateReason Reason, FString ReasonContext=TEXT(""));
 
 	/**
 	 * Recalculates cached counter history information, after an interval or state change which affects time period tracking.
@@ -475,17 +475,17 @@ private:
 	 * @param OutPerPeriodHistory		Outputs the recalculated counter cache, for the specified time periods (e.g. elements 1 and 7 only).
 	 * @param StartPerSecHistoryIdx		The index in CounterPerSecHistory to start calculating from
 	 */
-	void RecalculatePeriodHistory(const TArray<int8>& InTimePeriods, TArrayView<FEscalationCounter>(&OutPerPeriodHistory)[16],
+	NETCORE_API void RecalculatePeriodHistory(const TArray<int8>& InTimePeriods, TArrayView<FEscalationCounter>(&OutPerPeriodHistory)[16],
 									int32 StartPerSecHistoryIdx=INDEX_NONE);
 
-	int32 GetHighestHistoryRequirement() const;
+	NETCORE_API int32 GetHighestHistoryRequirement() const;
 
 	/**
 	 * Resets all counters - typically when entering dormancy
 	 */
-	void ResetAllCounters();
+	NETCORE_API void ResetAllCounters();
 
-	int32 AddNewCounter_Internal(int32 Count, const TArrayView<FEscalationCounter>& CountersPerPeriodAlloc);
+	NETCORE_API int32 AddNewCounter_Internal(int32 Count, const TArrayView<FEscalationCounter>& CountersPerPeriodAlloc);
 
 public:
 	/**
@@ -504,7 +504,7 @@ public:
 	/**
 	 * Checks if any escalation quota's have been hit, and if so, updates/escalates the severity state
 	 */
-	void CheckQuotas();
+	NETCORE_API void CheckQuotas();
 
 	/**
 	 * Ticks the escalation manager
@@ -512,7 +512,7 @@ public:
 	 *
 	 * @param TimeSeconds	The (can be approximate) current time
 	 */
-	void TickRealtime(double TimeSeconds);
+	NETCORE_API void TickRealtime(double TimeSeconds);
 
 	/**
 	 * Whether or not calls to Tick are presently required
@@ -530,22 +530,22 @@ public:
 	/**
 	 * Whether or not the current state is a dormant state, which does not require ticking
 	 */
-	bool IsDormant() const;
+	NETCORE_API bool IsDormant() const;
 
 	/**
 	 * Returns a reference to the BaseConfig object
 	 */
-	const UEscalationManagerConfig* GetBaseConfig() const;
+	NETCORE_API const UEscalationManagerConfig* GetBaseConfig() const;
 
 	/**
 	 * Sets a new ManagerContext value
 	 */
-	void SetManagerContext(FString InManagerContext);
+	NETCORE_API void SetManagerContext(FString InManagerContext);
 
 	/**
 	 * Sets a new NotifySeverityUpdate value
 	 */
-	void SetNotifySeverityUpdate(FNotifySeverityUpdate&& InNotifySeverityUpdate);
+	NETCORE_API void SetNotifySeverityUpdate(FNotifySeverityUpdate&& InNotifySeverityUpdate);
 
 
 private:
@@ -751,13 +751,13 @@ private:
  *
  * Subclass and override 'InitConfigDefaultsInternal' to initialize EscalationSeverity and bEnabled (and other custom config variables).
  */
-UCLASS(config=Engine, PerObjectConfig)
-class NETCORE_API UEscalationManagerConfig : public UStatePerObjectConfig
+UCLASS(config=Engine, PerObjectConfig, MinimalAPI)
+class UEscalationManagerConfig : public UStatePerObjectConfig
 {
 	GENERATED_BODY()
 
 private:
-	virtual void LoadStateConfig() override;
+	NETCORE_API virtual void LoadStateConfig() override;
 
 
 public:
diff --git a/Engine/Source/Runtime/Net/Core/Public/Net/Core/Connection/NetConnectionFaultRecoveryBase.h b/Engine/Source/Runtime/Net/Core/Public/Net/Core/Connection/NetConnectionFaultRecoveryBase.h
--- a/Engine/Source/Runtime/Net/Core/Public/Net/Core/Connection/NetConnectionFaultRecoveryBase.h
+++ b/Engine/Source/Runtime/Net/Core/Public/Net/Core/Connection/NetConnectionFaultRecoveryBase.h
@@ -36,7 +36,7 @@ namespace UE
  * Fault handlers may have their own separate escalation tracking.
  */
 USTRUCT()
-struct NETCORE_API FNetFaultState : public FEscalationState
+struct FNetFaultState : public FEscalationState
 {
 	GENERATED_BODY()
 
@@ -72,17 +72,17 @@ public:
 
 
 public:
-	static const TCHAR* GetConfigSection();
-	static UClass* GetBaseConfigClass();
+	static NETCORE_API const TCHAR* GetConfigSection();
+	static NETCORE_API UClass* GetBaseConfigClass();
 
 
 protected:
-	virtual EInitStateDefaultsResult InitConfigDefaultsInternal() override;
-	virtual void ApplyImpliedValuesInternal() override;
-	virtual void ValidateConfigInternal() override;
+	NETCORE_API virtual EInitStateDefaultsResult InitConfigDefaultsInternal() override;
+	NETCORE_API virtual void ApplyImpliedValuesInternal() override;
+	NETCORE_API virtual void ValidateConfigInternal() override;
 
 private:
-	virtual bool HasHitAnyQuota(FHasHitAnyQuotaParms Parms) const override;
+	NETCORE_API virtual bool HasHitAnyQuota(FHasHitAnyQuotaParms Parms) const override;
 };
 
 
@@ -147,7 +147,7 @@ inline uint8 ToInt(ENetFaultCounterCategory CategoryVal)
 /**
  * Implements the base/public interface for FNetConnectionFaultRecovery - defined here, to eliminate the need for Engine dependencies.
  */
-class NETCORE_API FNetConnectionFaultRecoveryBase
+class FNetConnectionFaultRecoveryBase
 {
 public:
 	using FNetFaultEscalationHandler = TEscalationManager<ENetFaultCounters, FNetFaultState, ENetFaultCounterCategory>;
@@ -162,7 +162,7 @@ public:
 	 * @param InResult	The successfully handled result, passed on to fault recovery
 	 * @return			Returns whether or not fault recovery handled the fault (this value should be returned by the calling HandleNetResult)
 	 */
-	EHandleNetResult NotifyHandledFault(FNetResult&& InResult);
+	NETCORE_API EHandleNetResult NotifyHandledFault(FNetResult&& InResult);
 
 
 	/** Accessors */
@@ -178,22 +178,22 @@ public:
 	/**
 	 * Convenience Forwarding function for FNetResultManager.HandleNetResult, taking FNetCloseResult/ENetCloseResult
 	 */
-	EHandleNetResult HandleNetResult(FNetCloseResult&& InResult);
+	NETCORE_API EHandleNetResult HandleNetResult(FNetCloseResult&& InResult);
 
 	/**
 	 * Passthrough for TEscalationManager.AddNewCounter (with caching, NetFaultEscalationManager has not been initialized yet)
 	 */
-	int32 AddNewCounter(int32 Count=1);
+	NETCORE_API int32 AddNewCounter(int32 Count=1);
 
 	/**
 	 * Passthrough for FEscalationManager.GetFrameCounter
 	 */
-	FEscalationCounter& GetFrameCounter(int32 CounterIndex);
+	NETCORE_API FEscalationCounter& GetFrameCounter(int32 CounterIndex);
 
 	/**
 	 * Passthrough for TEscalationManager.RegisterCounterCategory (with caching, NetFaultEscalationManager has not been initialized yet)
 	 */
-	void RegisterCounterCategory(ENetFaultCounterCategory Category, int32 CounterIndex);
+	NETCORE_API void RegisterCounterCategory(ENetFaultCounterCategory Category, int32 CounterIndex);
 
 
 private:
diff --git a/Engine/Source/Runtime/Net/Core/Public/Net/Core/Connection/NetResultManager.h b/Engine/Source/Runtime/Net/Core/Public/Net/Core/Connection/NetResultManager.h
--- a/Engine/Source/Runtime/Net/Core/Public/Net/Core/Connection/NetResultManager.h
+++ b/Engine/Source/Runtime/Net/Core/Public/Net/Core/Connection/NetResultManager.h
@@ -44,7 +44,7 @@ enum class EAddResultHandlerPos : uint8
  *
  * Handles arbitrary net results, which may attempt recovery from errors instead of e.g. immediately closing the NetConnection
  */
-class NETCORE_API FNetResultManager final
+class FNetResultManager final
 {
 public:
 	/**
@@ -69,7 +69,7 @@ public:
 	 * @param InResultHandler	The new result handler to be added
 	 * @param Position			The position/precedence in the result handler list, to place the new handler
 	 */
-	void AddResultHandler(TUniquePtr<FNetResultHandler>&& InResultHandler, EAddResultHandlerPos Position=EAddResultHandlerPos::Last);
+	NETCORE_API void AddResultHandler(TUniquePtr<FNetResultHandler>&& InResultHandler, EAddResultHandlerPos Position=EAddResultHandlerPos::Last);
 
 	/**
 	 * Adds a new result handler pointer to the result manager (not owned by the result manager)
@@ -77,7 +77,7 @@ public:
 	 * @param InResultHandler	The new result handler to be added
 	 * @param Position			The position/precedence in the result handler list, to place the new handler
 	 */
-	void AddResultHandlerPtr(FNetResultHandler* InResultHandler, EAddResultHandlerPos Position=EAddResultHandlerPos::Last);
+	NETCORE_API void AddResultHandlerPtr(FNetResultHandler* InResultHandler, EAddResultHandlerPos Position=EAddResultHandlerPos::Last);
 
 	/**
 	 * Takes a net result and passes it around to the result handlers and callbacks until it is handled, or returns 'EHandleNetResult::NotHandled'.
@@ -87,14 +87,14 @@ public:
 	 * @param InResult		Specifies the result
 	 * @return				Whether or not the result has handled or resulted in a close, or went unhandled
 	 */
-	EHandleNetResult HandleNetResult(FNetResult&& InResult);
+	NETCORE_API EHandleNetResult HandleNetResult(FNetResult&& InResult);
 
 	/**
 	 * Sets a callback for handling net results which no result handlers have dealt with
 	 *
 	 * @param InCallback	The callback to use for unhandled results
 	 */
-	void SetUnhandledResultCallback(FUnhandledResultFunc InCallback);
+	NETCORE_API void SetUnhandledResultCallback(FUnhandledResultFunc InCallback);
 
 private:
 	/** The list of result handlers, for attempting to handle results with */
@@ -154,4 +154,4 @@ protected:
 };
 
 }
-}
\ No newline at end of file
+}
