Fix navmesh rasterization issue when generating navmesh high above 0.

diff --git a/Engine/Source/Runtime/Navmesh/Private/Recast/RecastRasterization.cpp b/Engine/Source/Runtime/Navmesh/Private/Recast/RecastRasterization.cpp
--- a/Engine/Source/Runtime/Navmesh/Private/Recast/RecastRasterization.cpp
+++ b/Engine/Source/Runtime/Navmesh/Private/Recast/RecastRasterization.cpp
@@ -729,7 +729,7 @@ static void rasterizeTri(const rcReal* v0, const rcReal* v1, const rcReal* v2,
 			if (intverts[basevert][0] >= x0 && intverts[basevert][0] <= x1 && intverts[basevert][1] >= y0 && intverts[basevert][1] <= y1)
 			{
 				rcReal sfloat = vertarray[basevert][1] - bmin[1];
-				int sint = (int)rcClamp((int)rcFloor(sfloat * ich), -RANGE, RANGE);
+				const int sint = (int)rcClamp((int)rcFloor(sfloat * ich), -RANGE, RANGE);
 	#if TEST_NEW_RASTERIZER
 				rcAssert(sint >= triangle_ismin - 1 && sint <= triangle_ismax + 1);
 	#endif
@@ -767,7 +767,7 @@ static void rasterizeTri(const rcReal* v0, const rcReal* v1, const rcReal* v2,
 					if (y >= y0 && y <= y1)
 					{
 						rcReal sfloat = temppnt[1] - bmin[1];
-						short int sint = (short int)rcClamp((int)rcFloor(sfloat * ich), -RANGE, RANGE);
+						const int sint = (int)rcClamp((int)rcFloor(sfloat * ich), -RANGE, RANGE);
 #if TEST_NEW_RASTERIZER
 						rcAssert(sint >= triangle_ismin - 1 && sint <= triangle_ismax + 1);
 #endif
@@ -808,7 +808,7 @@ static void rasterizeTri(const rcReal* v0, const rcReal* v1, const rcReal* v2,
 						if (x >= x0 && x <= x1)
 						{
 							rcReal sfloat = Inter[i][1] - bmin[1];
-							short int sint = (short int)rcClamp((int)rcFloor(sfloat * ich), -RANGE, RANGE);
+							const int sint = (int)rcClamp((int)rcFloor(sfloat * ich), -RANGE, RANGE);
 #if TEST_NEW_RASTERIZER
 							rcAssert(sint >= triangle_ismin - 1 && sint <= triangle_ismax + 1);
 #endif
@@ -837,7 +837,7 @@ static void rasterizeTri(const rcReal* v0, const rcReal* v1, const rcReal* v2,
 						}
 						for (int x = xloop0; x <= xloop1; x++, sfloat += ds)
 						{
-							short int sint = (short int)rcClamp((int)rcFloor(sfloat * ich), -RANGE, RANGE);
+							const int sint = (int)rcClamp((int)rcFloor(sfloat * ich), -RANGE, RANGE);
 #if TEST_NEW_RASTERIZER
 							rcAssert(sint >= triangle_ismin - 1 && sint <= triangle_ismax + 1);
 #endif
