Fix spans being filtered out if they are over 0xffff verticaly

diff --git a/Engine/Source/Runtime/Navmesh/Private/Recast/RecastFilter.cpp b/Engine/Source/Runtime/Navmesh/Private/Recast/RecastFilter.cpp
--- a/Engine/Source/Runtime/Navmesh/Private/Recast/RecastFilter.cpp
+++ b/Engine/Source/Runtime/Navmesh/Private/Recast/RecastFilter.cpp
@@ -79,7 +79,7 @@ void rcFilterLedgeSpansImp(rcContext* ctx, const int walkableHeight, const int w
 
 	const int w = solid.width;
 	const int h = solid.height;
-	const int MAX_HEIGHT = 0xffff;
+	const int MAX_HEIGHT = RC_SPAN_MAX_HEIGHT;
 
 	// Mark border spans.
 	for (int x = 0; x < w; ++x)
@@ -181,9 +181,7 @@ void rcFilterLedgeSpans(rcContext* ctx, const int walkableHeight, const int walk
 	
 	ctx->startTimer(RC_TIMER_FILTER_BORDER);
 
-	const int w = solid.width;
 	const int h = solid.height;
-	const int MAX_HEIGHT = 0xffff;
 	
 	// Mark border spans.
 	for (int y = 0; y < h; ++y)
@@ -203,7 +201,6 @@ void rcFilterLedgeSpans(rcContext* ctx, const int walkableHeight, const int walk
 
 	ctx->startTimer(RC_TIMER_FILTER_BORDER);
 
-	const int w = solid.width;
 	const int h = rcMin(yStart + maxYProcess, solid.height);
 
 	for (int y = yStart; y < h; ++y)
@@ -228,7 +225,7 @@ void rcFilterWalkableLowHeightSpans(rcContext* ctx, int walkableHeight, rcHeight
 	
 	const int w = solid.width;
 	const int h = solid.height;
-	const int MAX_HEIGHT = 0xffff;
+	const int MAX_HEIGHT = RC_SPAN_MAX_HEIGHT;
 	
 	// Remove walkable flag from spans which do not have enough
 	// space above them for the agent to stand there.
