Properly checking for extents in navmesh findNearestPoly queries.

diff --git a/Engine/Source/Runtime/Navmesh/Private/Detour/DetourNavMeshQuery.cpp b/Engine/Source/Runtime/Navmesh/Private/Detour/DetourNavMeshQuery.cpp
--- a/Engine/Source/Runtime/Navmesh/Private/Detour/DetourNavMeshQuery.cpp
+++ b/Engine/Source/Runtime/Navmesh/Private/Detour/DetourNavMeshQuery.cpp
@@ -1142,10 +1142,16 @@ dtStatus dtNavMeshQuery::findNearestPoly(const dtReal* center, const dtReal* ext
 //@UE END
 		if (d < nearestDistanceSqr && h < extents[1])
 		{
-			if (nearestPt)
-				dtVcopy(nearestPt, closestPtPoly);
-			nearestDistanceSqr = d;
-			nearest = ref;
+			// Checking horizontal extents except when using a referencePt since that can lead to missing results.
+			const bool insideX = referencePt || dtAbs(center[0] - closestPtPoly[0]) < extents[0];
+			const bool insideZ = referencePt || dtAbs(center[2] - closestPtPoly[2]) < extents[2];
+			if (insideX && insideZ)
+			{
+				if (nearestPt)
+					dtVcopy(nearestPt, closestPtPoly);
+				nearestDistanceSqr = d;
+				nearest = ref;
+			}
 		}
 	}
 	
@@ -1192,6 +1198,15 @@ dtStatus dtNavMeshQuery::findNearestPoly2D(const dtReal* center, const dtReal* e
 
 		if (h > extents[1])
 			continue;
+
+		// If we are not using a referencePt, check extents (else it can lead to missing results).
+		if (!referencePt)
+		{
+			const bool outsideX = dtAbs(center[0] - closestPtPoly[0]) > extents[0];
+			const bool outsideZ = dtAbs(center[2] - closestPtPoly[2]) > extents[2];
+			if (outsideX || outsideZ)
+				continue;
+		}
 		
 		// scoring depends if 2D distance to center is within tolerance value
 		if (dSq < toleranceSq)
@@ -1264,9 +1279,12 @@ dtStatus dtNavMeshQuery::findNearestContainingPoly(const dtReal* center, const d
 		{
 			dtReal closestPtPoly[3];
 			closestPointOnPoly(ref, center, closestPtPoly);
-			dtReal d = dtVdistSqr(center, closestPtPoly);
-			dtReal h = dtAbs(center[1] - closestPtPoly[1]);
-			if (d < nearestDistanceSqr && h < extents[1])
+			const dtReal d = dtVdistSqr(center, closestPtPoly);
+			const dtReal x = dtAbs(center[0] - closestPtPoly[0]);
+			const dtReal z = dtAbs(center[2] - closestPtPoly[2]);
+			const dtReal h = dtAbs(center[1] - closestPtPoly[1]);
+
+			if (d < nearestDistanceSqr && h < extents[1] && x < extents[0] && z < extents[2])
 			{
 				if (nearestPt)
 					dtVcopy(nearestPt, closestPtPoly);
@@ -1304,9 +1322,12 @@ dtPolyRef dtNavMeshQuery::findNearestPolyInTile(const dtMeshTile* tile, const dt
 		const dtPoly* poly = &tile->polys[m_nav->decodePolyIdPoly(ref)];
 		dtReal closestPtPoly[3];
 		closestPointOnPolyInTile(tile, poly, center, closestPtPoly);
-			
-		dtReal d = dtVdistSqr(center, closestPtPoly);
-		if (d < nearestDistanceSqr)
+		const dtReal d = dtVdistSqr(center, closestPtPoly);
+		const dtReal x = dtAbs(center[0] - closestPtPoly[0]);
+		const dtReal z = dtAbs(center[2] - closestPtPoly[2]);
+		const dtReal h = dtAbs(center[1] - closestPtPoly[1]);
+
+		if (d < nearestDistanceSqr && h < extents[1] && x < extents[0] && z < extents[2])
 		{
 			if (nearestPt)
 				dtVcopy(nearestPt, closestPtPoly);
diff --git a/Engine/Source/Runtime/Navmesh/Public/Detour/DetourNavMeshQuery.h b/Engine/Source/Runtime/Navmesh/Public/Detour/DetourNavMeshQuery.h
--- a/Engine/Source/Runtime/Navmesh/Public/Detour/DetourNavMeshQuery.h
+++ b/Engine/Source/Runtime/Navmesh/Public/Detour/DetourNavMeshQuery.h
@@ -497,7 +497,7 @@ public:
 	///  @param[in]		filter		The polygon filter to apply to the query.
 	///  @param[out]	nearestRef	The reference id of the nearest polygon.
 	///  @param[out]	nearestPt	The nearest point on the polygon. [opt] [(x, y, z)]
-	///  @param[in]		referencePt	If supplied replaces @param center in terms of distance measurements. [opt] [(x, y, z)]
+	///  @param[in]		referencePt	If supplied replaces @param center in terms of distance measurements. Note that extents are not precise when using a referencePt. [opt] [(x, y, z)]
 	/// @returns The status flags for the query.
 	NAVMESH_API dtStatus findNearestPoly(const dtReal* center, const dtReal* extents,
 							 const dtQueryFilter* filter,
@@ -509,7 +509,7 @@ public:
 	///  @param[in]		filter		The polygon filter to apply to the query.
 	///  @param[out]	nearestRef	The reference id of the nearest polygon.
 	///  @param[out]	nearestPt	The nearest point on the polygon. [opt] [(x, y, z)]
-	///  @param[in]		referencePt	If supplied replaces @param center in terms of distance measurements. [opt] [(x, y, z)]
+	///  @param[in]		referencePt	If supplied replaces @param center in terms of distance measurements. Note that extents are not precise when using a referencePt. [opt] [(x, y, z)]
 	///  @param[in]		tolerance	Radius around best 2D point for picking vertical match
 	/// @returns The status flags for the query.
 	NAVMESH_API dtStatus findNearestPoly2D(const dtReal* center, const dtReal* extents,
