Allow navmesh queries with 0 extent to succeed since that was the behavior before the change to check the extent precisely.

diff --git a/Engine/Source/Runtime/Navmesh/Private/Detour/DetourNavMeshQuery.cpp b/Engine/Source/Runtime/Navmesh/Private/Detour/DetourNavMeshQuery.cpp
--- a/Engine/Source/Runtime/Navmesh/Private/Detour/DetourNavMeshQuery.cpp
+++ b/Engine/Source/Runtime/Navmesh/Private/Detour/DetourNavMeshQuery.cpp
@@ -1140,11 +1140,11 @@ dtStatus dtNavMeshQuery::findNearestPoly(const dtReal* center, const dtReal* ext
 		const dtReal d = dtVdistSqr(referenceLocation, closestPtPoly);
 		const dtReal h = dtAbs(center[1] - closestPtPoly[1]);
 //@UE END
-		if (d < nearestDistanceSqr && h < extents[1])
+		if (d < nearestDistanceSqr && h <= extents[1])
 		{
 			// Checking horizontal extents except when using a referencePt since that can lead to missing results.
-			const bool insideX = referencePt || dtAbs(center[0] - closestPtPoly[0]) < extents[0];
-			const bool insideZ = referencePt || dtAbs(center[2] - closestPtPoly[2]) < extents[2];
+			const bool insideX = referencePt || dtAbs(center[0] - closestPtPoly[0]) <= extents[0];
+			const bool insideZ = referencePt || dtAbs(center[2] - closestPtPoly[2]) <= extents[2];
 			if (insideX && insideZ)
 			{
 				if (nearestPt)
@@ -1196,14 +1196,14 @@ dtStatus dtNavMeshQuery::findNearestPoly2D(const dtReal* center, const dtReal* e
 		const dtReal dSq = dtVdist2DSqr(referenceLocation, closestPtPoly);
 		const dtReal h = dtAbs(center[1] - closestPtPoly[1]);
 
-		if (h > extents[1])
+		if (h >= extents[1])
 			continue;
 
 		// If we are not using a referencePt, check extents (else it can lead to missing results).
 		if (!referencePt)
 		{
-			const bool outsideX = dtAbs(center[0] - closestPtPoly[0]) > extents[0];
-			const bool outsideZ = dtAbs(center[2] - closestPtPoly[2]) > extents[2];
+			const bool outsideX = dtAbs(center[0] - closestPtPoly[0]) >= extents[0];
+			const bool outsideZ = dtAbs(center[2] - closestPtPoly[2]) >= extents[2];
 			if (outsideX || outsideZ)
 				continue;
 		}
@@ -1284,7 +1284,7 @@ dtStatus dtNavMeshQuery::findNearestContainingPoly(const dtReal* center, const d
 			const dtReal z = dtAbs(center[2] - closestPtPoly[2]);
 			const dtReal h = dtAbs(center[1] - closestPtPoly[1]);
 
-			if (d < nearestDistanceSqr && h < extents[1] && x < extents[0] && z < extents[2])
+			if (d < nearestDistanceSqr && h <= extents[1] && x <= extents[0] && z <= extents[2])
 			{
 				if (nearestPt)
 					dtVcopy(nearestPt, closestPtPoly);
@@ -1327,7 +1327,7 @@ dtPolyRef dtNavMeshQuery::findNearestPolyInTile(const dtMeshTile* tile, const dt
 		const dtReal z = dtAbs(center[2] - closestPtPoly[2]);
 		const dtReal h = dtAbs(center[1] - closestPtPoly[1]);
 
-		if (d < nearestDistanceSqr && h < extents[1] && x < extents[0] && z < extents[2])
+		if (d < nearestDistanceSqr && h <= extents[1] && x <= extents[0] && z <= extents[2])
 		{
 			if (nearestPt)
 				dtVcopy(nearestPt, closestPtPoly);
