Allow navmesh queries with 0 extent to succeed when using findNearestPoly2D.

diff --git a/Engine/Source/Runtime/Navmesh/Private/Detour/DetourNavMeshQuery.cpp b/Engine/Source/Runtime/Navmesh/Private/Detour/DetourNavMeshQuery.cpp
--- a/Engine/Source/Runtime/Navmesh/Private/Detour/DetourNavMeshQuery.cpp
+++ b/Engine/Source/Runtime/Navmesh/Private/Detour/DetourNavMeshQuery.cpp
@@ -1196,14 +1196,16 @@ dtStatus dtNavMeshQuery::findNearestPoly2D(const dtReal* center, const dtReal* e
 		const dtReal dSq = dtVdist2DSqr(referenceLocation, closestPtPoly);
 		const dtReal h = dtAbs(center[1] - closestPtPoly[1]);
 
-		if (h >= extents[1])
+		if (h > extents[1])
 			continue;
 
 		// If we are not using a referencePt, check extents (else it can lead to missing results).
 		if (!referencePt)
 		{
-			const bool outsideX = dtAbs(center[0] - closestPtPoly[0]) >= extents[0];
-			const bool outsideZ = dtAbs(center[2] - closestPtPoly[2]) >= extents[2];
+			const double deltaX = dtAbs(center[0] - closestPtPoly[0]);
+			const double deltaZ = dtAbs(center[2] - closestPtPoly[2]);
+			const bool outsideX = deltaX > extents[0];	// Using > instead of >= to allow 0 extent to be valid.
+			const bool outsideZ = deltaZ > extents[2];
 			if (outsideX || outsideZ)
 				continue;
 		}
