Navmesh queries: clamp extents to a minimum epsilon value to handle extents with components set to 0.

diff --git a/Engine/Source/Runtime/Navmesh/Private/Detour/DetourNavMeshQuery.cpp b/Engine/Source/Runtime/Navmesh/Private/Detour/DetourNavMeshQuery.cpp
--- a/Engine/Source/Runtime/Navmesh/Private/Detour/DetourNavMeshQuery.cpp
+++ b/Engine/Source/Runtime/Navmesh/Private/Detour/DetourNavMeshQuery.cpp
@@ -994,6 +994,15 @@ dtStatus dtNavMeshQuery::projectedPointOnPolyInTile(const dtMeshTile* tile, cons
 	return DT_FAILURE;
 }
 
+//@UE BEGIN
+void dtNavMeshQuery::dtApplyEpsilon(dtReal* extents)
+{
+	constexpr dtReal EPS = 1.0e-6;
+	dtReal epsilonVector[3] = { EPS, EPS, EPS };
+	dtVmax(extents, epsilonVector);
+}
+//@UE END
+
 /// @par
 ///
 /// Will return #DT_FAILURE if the provided position is outside the xz-bounds 
@@ -1118,11 +1127,19 @@ dtStatus dtNavMeshQuery::findNearestPoly(const dtReal* center, const dtReal* ext
 	dtAssert(m_nav);
 
 	*nearestRef = 0;
+
+	//@UE BEGIN
+	// Make sure the extents are at least epsilon in size. Zero extents (0,y,0) are sometimes used as "ray cast down".
+	// The epsilon is there to prevent the query to miss results when the query is right on the (internal) edge of a polygon.
+	dtReal localExtents[3];
+	dtVcopy(localExtents, extents);
+	dtApplyEpsilon(localExtents);
+	//@UE END
 	
 	// Get nearby polygons from proximity grid.
 	dtPolyRef polys[128];
 	int polyCount = 0;
-	if (dtStatusFailed(queryPolygons(center, extents, filter, polys, &polyCount, 128)))
+	if (dtStatusFailed(queryPolygons(center, localExtents, filter, polys, &polyCount, 128)))
 		return DT_FAILURE | DT_INVALID_PARAM;
 	
 //@UE BEGIN
@@ -1140,11 +1157,11 @@ dtStatus dtNavMeshQuery::findNearestPoly(const dtReal* center, const dtReal* ext
 		const dtReal d = dtVdistSqr(referenceLocation, closestPtPoly);
 		const dtReal h = dtAbs(center[1] - closestPtPoly[1]);
 //@UE END
-		if (d < nearestDistanceSqr && h <= extents[1])
+		if (d < nearestDistanceSqr && h <= localExtents[1])
 		{
 			// Checking horizontal extents except when using a referencePt since that can lead to missing results.
-			const bool insideX = referencePt || dtAbs(center[0] - closestPtPoly[0]) <= extents[0];
-			const bool insideZ = referencePt || dtAbs(center[2] - closestPtPoly[2]) <= extents[2];
+			const bool insideX = referencePt || dtAbs(center[0] - closestPtPoly[0]) <= localExtents[0];	//@UE
+			const bool insideZ = referencePt || dtAbs(center[2] - closestPtPoly[2]) <= localExtents[2];	//@UE
 			if (insideX && insideZ)
 			{
 				if (nearestPt)
@@ -1171,10 +1188,16 @@ dtStatus dtNavMeshQuery::findNearestPoly2D(const dtReal* center, const dtReal* e
 
 	*nearestRef = 0;
 
+	// Make sure the extents are at least epsilon in size. Zero extents (0,y,0) are sometimes used as "ray cast down".
+	// The epsilon is there to prevent the query to miss results when the query is right on the (internal) edge of a polygon.
+	dtReal localExtents[3];
+	dtVcopy(localExtents, extents);
+	dtApplyEpsilon(localExtents);
+	
 	// Get nearby polygons from proximity grid.
 	dtPolyRef polys[128];
 	int polyCount = 0;
-	if (dtStatusFailed(queryPolygons(center, extents, filter, polys, &polyCount, 128)))
+	if (dtStatusFailed(queryPolygons(center, localExtents, filter, polys, &polyCount, 128)))
 		return DT_FAILURE | DT_INVALID_PARAM;
 
 	const dtReal toleranceSq = dtSqr(tolerance);
@@ -1196,7 +1219,7 @@ dtStatus dtNavMeshQuery::findNearestPoly2D(const dtReal* center, const dtReal* e
 		const dtReal dSq = dtVdist2DSqr(referenceLocation, closestPtPoly);
 		const dtReal h = dtAbs(center[1] - closestPtPoly[1]);
 
-		if (h > extents[1])
+		if (h > localExtents[1])
 			continue;
 
 		// If we are not using a referencePt, check extents (else it can lead to missing results).
@@ -1204,8 +1227,8 @@ dtStatus dtNavMeshQuery::findNearestPoly2D(const dtReal* center, const dtReal* e
 		{
 			const double deltaX = dtAbs(center[0] - closestPtPoly[0]);
 			const double deltaZ = dtAbs(center[2] - closestPtPoly[2]);
-			const bool outsideX = deltaX > extents[0];	// Using > instead of >= to allow 0 extent to be valid.
-			const bool outsideZ = deltaZ > extents[2];
+			const bool outsideX = deltaX > localExtents[0];
+			const bool outsideZ = deltaZ > localExtents[2];
 			if (outsideX || outsideZ)
 				continue;
 		}
@@ -1261,10 +1284,18 @@ dtStatus dtNavMeshQuery::findNearestContainingPoly(const dtReal* center, const d
 
 	*nearestRef = 0;
 
+	//@UE BEGIN
+	// Make sure the extents are at least epsilon in size. Zero extents (0,y,0) are sometimes used as "ray cast down".
+	// The epsilon is there to prevent the query to miss results when the query is right on the (internal) edge of a polygon.
+	dtReal localExtents[3];
+	dtVcopy(localExtents, extents);
+	dtApplyEpsilon(localExtents);
+	//@UE END
+	
 	// Get nearby polygons from proximity grid.
 	dtPolyRef polys[128];
 	int polyCount = 0;
-	if (dtStatusFailed(queryPolygons(center, extents, filter, polys, &polyCount, 128)))
+	if (dtStatusFailed(queryPolygons(center, localExtents, filter, polys, &polyCount, 128)))
 		return DT_FAILURE | DT_INVALID_PARAM;
 
 	// Find nearest polygon amongst the nearby polygons.
@@ -1286,7 +1317,7 @@ dtStatus dtNavMeshQuery::findNearestContainingPoly(const dtReal* center, const d
 			const dtReal z = dtAbs(center[2] - closestPtPoly[2]);
 			const dtReal h = dtAbs(center[1] - closestPtPoly[1]);
 
-			if (d < nearestDistanceSqr && h <= extents[1] && x <= extents[0] && z <= extents[2])
+			if (d < nearestDistanceSqr && h <= localExtents[1] && x <= localExtents[0] && z <= localExtents[2]) //@UE
 			{
 				if (nearestPt)
 					dtVcopy(nearestPt, closestPtPoly);
@@ -1306,10 +1337,18 @@ dtPolyRef dtNavMeshQuery::findNearestPolyInTile(const dtMeshTile* tile, const dt
 												const dtQueryFilter* filter, dtReal* nearestPt) const
 {
 	dtAssert(m_nav);
+
+	//@UE BEGIN
+	// Make sure the extents are at least epsilon in size. Zero extents (0,y,0) are sometimes used as "ray cast down".
+	// The epsilon is there to prevent the query to miss results when the query is right on the (internal) edge of a polygon.
+	dtReal localExtents[3];
+	dtVcopy(localExtents, extents);
+	dtApplyEpsilon(localExtents);
+	//@UE END
 	
 	dtReal bmin[3], bmax[3];
-	dtVsub(bmin, center, extents);
-	dtVadd(bmax, center, extents);
+	dtVsub(bmin, center, localExtents);
+	dtVadd(bmax, center, localExtents);
 	
 	// Get nearby polygons from proximity grid.
 	dtPolyRef polys[128];
@@ -1329,7 +1368,7 @@ dtPolyRef dtNavMeshQuery::findNearestPolyInTile(const dtMeshTile* tile, const dt
 		const dtReal z = dtAbs(center[2] - closestPtPoly[2]);
 		const dtReal h = dtAbs(center[1] - closestPtPoly[1]);
 
-		if (d < nearestDistanceSqr && h <= extents[1] && x <= extents[0] && z <= extents[2])
+		if (d < nearestDistanceSqr && h <= localExtents[1] && x <= localExtents[0] && z <= localExtents[2])	//@UE
 		{
 			if (nearestPt)
 				dtVcopy(nearestPt, closestPtPoly);
diff --git a/Engine/Source/Runtime/Navmesh/Public/Detour/DetourNavMeshQuery.h b/Engine/Source/Runtime/Navmesh/Public/Detour/DetourNavMeshQuery.h
--- a/Engine/Source/Runtime/Navmesh/Public/Detour/DetourNavMeshQuery.h
+++ b/Engine/Source/Runtime/Navmesh/Public/Detour/DetourNavMeshQuery.h
@@ -778,6 +778,9 @@ private:
 	NAVMESH_API dtStatus projectedPointOnPolyInTile(const dtMeshTile* tile, const dtPoly* poly, const dtReal* pos, dtReal* projected) const;
 	
 	//@UE BEGIN
+	/// Handle 0 extents by having the vector components set to a minimum value.
+	static void dtApplyEpsilon(dtReal* extents);
+	
 	// exposing function to be able to generate navigation corridors as sequence of point pairs
 public:
 	/// Returns portal points between two polygons.
