Fix navmesh being generated a locations too steep and above the step height threshold.
Addition of LedgeFilterMode on RecastNavMesh to select the desired ledge slope filter mode.
Addition of a new filter mode that takes into acount the maximum walkable slope when filtering voxels.

diff --git a/Engine/Source/Runtime/NavigationSystem/Private/NavMesh/RecastNavMesh.cpp b/Engine/Source/Runtime/NavigationSystem/Private/NavMesh/RecastNavMesh.cpp
--- a/Engine/Source/Runtime/NavigationSystem/Private/NavMesh/RecastNavMesh.cpp
+++ b/Engine/Source/Runtime/NavigationSystem/Private/NavMesh/RecastNavMesh.cpp
@@ -542,6 +542,7 @@ ARecastNavMesh::ARecastNavMesh(const FObjectInitializer& ObjectInitializer)
 	, RecastNavMeshImpl(NULL)
 {
 	HeuristicScale = 0.999f;
+	LedgeSlopeFilterMode = ENavigationLedgeSlopeFilterMode::Recast;
 	RegionPartitioning = ERecastPartitioning::Watershed;
 	LayerPartitioning = ERecastPartitioning::Watershed;
 	RegionChunkSplits = 2;
diff --git a/Engine/Source/Runtime/NavigationSystem/Private/NavMesh/RecastNavMeshGenerator.cpp b/Engine/Source/Runtime/NavigationSystem/Private/NavMesh/RecastNavMeshGenerator.cpp
--- a/Engine/Source/Runtime/NavigationSystem/Private/NavMesh/RecastNavMeshGenerator.cpp
+++ b/Engine/Source/Runtime/NavigationSystem/Private/NavMesh/RecastNavMeshGenerator.cpp
@@ -82,9 +82,6 @@ namespace UE::NavMesh::Private
 
 	static bool bUseTightBoundExpansion = true;
 	static FAutoConsoleVariableRef CVarUseTightBoundExpansion(TEXT("ai.nav.UseTightBoundExpansion"), bUseTightBoundExpansion, TEXT("Active by default. Use an expansion of one AgentRadius. Set to false to revert to the previous behavior (2 AgentRadius)."), ECVF_Default);
-
-	static bool bKeepSteepSlopeForSingleVoxelAgent = true;
-	static FAutoConsoleVariableRef CVarKeepSteepSlopeForSingleVoxelAgent(TEXT("ai.nav.KeepSteepSlopeForSingleVoxelAgent"), bKeepSteepSlopeForSingleVoxelAgent, TEXT("Active by default. Fix too wide filtering of steep slope in Recast Heightfield filtering when the agent radius is only 1 voxel. Set to false to revert to the previous behavior."), ECVF_Default);
 }
 
 static FOodleDataCompression::ECompressor GNavmeshTileCacheCompressor = FOodleDataCompression::ECompressor::Mermaid;
@@ -250,16 +247,6 @@ struct FRecastGeometryExport : public FNavigableGeometryExport
 	virtual void SetNavDataPerInstanceTransformDelegate(const FNavDataPerInstanceTransformDelegate& InDelegate) override;
 };
 
-bool FRecastBuildConfig::IsUsingCoarseCellSize() const
-{
-	const rcReal MaxStepFromUniformSlope = cs * FMath::Tan(FMath::DegreesToRadians(walkableSlopeAngle));
-
-	// When filtering occurs, in rcFilterLedgeSpansImp, there can be a 2 steps vertical distance between neighbors ((asmax - asmin) > walkableClimb).
-	// This is why we compare with 2 times the max step.
-	const bool bIsUsingCoarseCellSize = (walkableClimb * ch) <= (2 * MaxStepFromUniformSlope);
-	return bIsUsingCoarseCellSize;
-}
-
 FRecastVoxelCache::FRecastVoxelCache(const uint8* Memory)
 {
 	uint8* BytesArr = (uint8*)Memory;
@@ -2966,9 +2953,9 @@ void FRecastTileGenerator::GenerateRecastFilter(FNavMeshBuildContext& BuildConte
 	}
 	{
 		SCOPE_CYCLE_COUNTER(STAT_Navigation_FilterLedgeSpans)
-
-		const bool bFilterNeighborSlope = (TileConfig.walkableRadius > 1 || !UE::NavMesh::Private::bKeepSteepSlopeForSingleVoxelAgent) && !TileConfig.IsUsingCoarseCellSize();
-		rcFilterLedgeSpans(&BuildContext, TileConfig.walkableHeight, TileConfig.walkableClimb, bFilterNeighborSlope, *RasterContext.SolidHF);
+		
+		rcFilterLedgeSpans(&BuildContext, TileConfig.walkableHeight, TileConfig.walkableClimb,
+			(rcNeighborSlopeFilterMode)TileConfig.LedgeSlopeFilterMode, TileConfig.maxStepFromWalkableSlope, TileConfig.ch, *RasterContext.SolidHF);
 	}
 	if (!TileConfig.bMarkLowHeightAreas)
 	{
@@ -3014,11 +3001,11 @@ ETimeSliceWorkResult FRecastTileGenerator::GenerateRecastFilterTimeSliced(FNavMe
 
 		bool DoIter = true;
 
-		const bool bFilterNeighborSlope = (TileConfig.walkableRadius > 1 || !UE::NavMesh::Private::bKeepSteepSlopeForSingleVoxelAgent) && !TileConfig.IsUsingCoarseCellSize();
-
 		do
 		{
-			rcFilterLedgeSpans(&BuildContext, TileConfig.walkableHeight, TileConfig.walkableClimb, bFilterNeighborSlope, GenRecastFilterLedgeSpansYStart, TileTimeSliceSettings.FilterLedgeSpansMaxYProcess, *RasterContext.SolidHF);
+			rcFilterLedgeSpans(&BuildContext, TileConfig.walkableHeight, TileConfig.walkableClimb,
+				(rcNeighborSlopeFilterMode)TileConfig.LedgeSlopeFilterMode, TileConfig.maxStepFromWalkableSlope, TileConfig.ch, 
+				GenRecastFilterLedgeSpansYStart, TileTimeSliceSettings.FilterLedgeSpansMaxYProcess, *RasterContext.SolidHF);
 
 			GenRecastFilterLedgeSpansYStart += TileTimeSliceSettings.FilterLedgeSpansMaxYProcess;
 
@@ -4638,6 +4625,7 @@ void FRecastNavMeshGenerator::SetupTileConfig(const ENavigationDataResolution Ti
 	OutConfig.TileResolution = TileResolution;
 	OutConfig.cs = CellSize;
 	OutConfig.walkableRadius = FMath::CeilToInt(DestNavMesh->AgentRadius / CellSize);
+	OutConfig.maxStepFromWalkableSlope = OutConfig.cs * FMath::Tan(FMath::DegreesToRadians(OutConfig.walkableSlopeAngle));
 
 	OutConfig.borderSize = OutConfig.walkableRadius + 3; // +1 for voxelization rounding, +1 for ledge neighbor access, +1 for occasional errors
 	OutConfig.maxEdgeLen = (int32)(1200.0f / CellSize);
@@ -4679,7 +4667,8 @@ void FRecastNavMeshGenerator::ConfigureBuildProperties(FRecastBuildConfig& OutCo
 	OutConfig.walkableHeight = FMath::CeilToInt(AgentHeight / CellHeight);
 	OutConfig.walkableClimb = FMath::CeilToInt(AgentMaxClimb / CellHeight);
 	OutConfig.walkableRadius = FMath::CeilToInt(AgentRadius / CellSize);
-
+	OutConfig.maxStepFromWalkableSlope = OutConfig.cs * FMath::Tan(FMath::DegreesToRadians(OutConfig.walkableSlopeAngle));
+	
 	// For each navmesh resolutions, validate that AgentMaxStepHeight is high enough for the AgentMaxSlope angle
 	for (int32 Index = 0; Index < (uint8)ENavigationDataResolution::MAX; Index++)
 	{
@@ -4743,6 +4732,7 @@ void FRecastNavMeshGenerator::ConfigureBuildProperties(FRecastBuildConfig& OutCo
 
 	OutConfig.regionChunkSize = FMath::Max(1, OutConfig.tileSize / FMath::Max(1, DestNavMesh->LayerChunkSplits));
 	OutConfig.TileCacheChunkSize = FMath::Max(1, OutConfig.tileSize / FMath::Max(1, DestNavMesh->RegionChunkSplits));
+	OutConfig.LedgeSlopeFilterMode = DestNavMesh->LedgeSlopeFilterMode;
 	OutConfig.regionPartitioning = DestNavMesh->LayerPartitioning;
 	OutConfig.TileCachePartitionType = DestNavMesh->RegionPartitioning;
 }
diff --git a/Engine/Source/Runtime/NavigationSystem/Public/NavMesh/RecastNavMesh.h b/Engine/Source/Runtime/NavigationSystem/Public/NavMesh/RecastNavMesh.h
--- a/Engine/Source/Runtime/NavigationSystem/Public/NavMesh/RecastNavMesh.h
+++ b/Engine/Source/Runtime/NavigationSystem/Public/NavMesh/RecastNavMesh.h
@@ -71,6 +71,14 @@ namespace ERecastPartitioning
 	};
 }
 
+UENUM()
+enum class ENavigationLedgeSlopeFilterMode : uint8
+{
+	Recast,							// Use walkableClimb value to filter
+	None,							// Skip slope filtering
+	UseStepHeightFromAgentMaxSlope	// Use maximum step height computed from AgentMaxSlope
+};
+
 struct FDetourTileSizeInfo
 {
 	unsigned short VertCount = 0;
@@ -371,6 +379,10 @@ struct FRecastNavMeshGenerationProperties
 	UPROPERTY(EditAnywhere, Category = Generation, meta = (ClampMin = "1", UIMin = "1"), AdvancedDisplay)
 	int32 TileNumberHardLimit;
 
+	/** filtering methode used for filtering ledge slopes */
+	UPROPERTY(EditAnywhere, Category = Generation, AdvancedDisplay)
+	ENavigationLedgeSlopeFilterMode LedgeFilterMode;
+	
 	/** partitioning method for creating navmesh polys */
 	UPROPERTY(EditAnywhere, Category = Generation, AdvancedDisplay)
 	TEnumAsByte<ERecastPartitioning::Type> RegionPartitioning;
@@ -827,6 +839,10 @@ class ARecastNavMesh : public ANavigationData
 	UPROPERTY(config)
 	float DefaultMaxHierarchicalSearchNodes;
 
+	/** filtering methode used for filtering ledge slopes */
+	UPROPERTY(EditAnywhere, Category=Generation, config, AdvancedDisplay)
+	ENavigationLedgeSlopeFilterMode LedgeSlopeFilterMode;
+	
 	/** partitioning method for creating navmesh polys */
 	UPROPERTY(EditAnywhere, Category=Generation, config, AdvancedDisplay)
 	TEnumAsByte<ERecastPartitioning::Type> RegionPartitioning;
diff --git a/Engine/Source/Runtime/NavigationSystem/Public/NavMesh/RecastNavMeshGenerator.h b/Engine/Source/Runtime/NavigationSystem/Public/NavMesh/RecastNavMeshGenerator.h
--- a/Engine/Source/Runtime/NavigationSystem/Public/NavMesh/RecastNavMeshGenerator.h
+++ b/Engine/Source/Runtime/NavigationSystem/Public/NavMesh/RecastNavMeshGenerator.h
@@ -81,6 +81,8 @@ struct FRecastBuildConfig : public rcConfig
 	int32 AgentIndex;
 	/** Resolution level */ 
 	ENavigationDataResolution TileResolution;
+	/** Ledge filtering mode */
+	ENavigationLedgeSlopeFilterMode LedgeSlopeFilterMode;
 
 	FRecastBuildConfig()
 	{
@@ -101,6 +103,7 @@ struct FRecastBuildConfig : public rcConfig
 		MaxPolysPerTile = -1;
 		AgentIndex = 0;
 		TileResolution = ENavigationDataResolution::Default;
+		LedgeSlopeFilterMode = ENavigationLedgeSlopeFilterMode::Recast;
 	}
 
 	rcReal GetTileSizeUU() const { return tileSize * cs; }
diff --git a/Engine/Source/Runtime/Navmesh/Private/Recast/RecastFilter.cpp b/Engine/Source/Runtime/Navmesh/Private/Recast/RecastFilter.cpp
--- a/Engine/Source/Runtime/Navmesh/Private/Recast/RecastFilter.cpp
+++ b/Engine/Source/Runtime/Navmesh/Private/Recast/RecastFilter.cpp
@@ -72,8 +72,9 @@ void rcFilterLowHangingWalkableObstacles(rcContext* ctx, const int walkableClimb
 	ctx->stopTimer(RC_TIMER_FILTER_LOW_OBSTACLES);
 }
 
-void rcFilterLedgeSpansImp(rcContext* ctx, const int walkableHeight, const int walkableClimb, const bool filterNeighborSlope, const int filterLedgeSpansAtY, //UE
-	rcHeightfield& solid)
+void rcFilterLedgeSpansImp(rcContext* ctx, const int walkableHeight, const int walkableClimb,
+                           const rcNeighborSlopeFilterMode neighborSlopeFilterMode, const rcReal maxStepFromWalkableSlope, const rcReal ch,	// UE
+                           const int filterLedgeSpansAtY, rcHeightfield& solid)
 {
 	rcAssert(ctx);
 
@@ -81,6 +82,8 @@ void rcFilterLedgeSpansImp(rcContext* ctx, const int walkableHeight, const int w
 	const int h = solid.height;
 	const int MAX_HEIGHT = RC_SPAN_MAX_HEIGHT;
 
+	const int maxStepFor2CellsVx = rcCeil(2*maxStepFromWalkableSlope / ch);	// UE
+
 	// Mark border spans.
 	for (int x = 0; x < w; ++x)
 	{
@@ -130,12 +133,11 @@ void rcFilterLedgeSpansImp(rcContext* ctx, const int walkableHeight, const int w
 						minh = rcMin(minh, nbot - bot);
 
 						// Find min/max accessible neighbour height. 
-						if (filterNeighborSlope && rcAbs(nbot - bot) <= walkableClimb)	//UE
+						if (neighborSlopeFilterMode != RC_SLOPE_FILTER_NONE && rcAbs(nbot - bot) <= walkableClimb)	//UE
 						{
 							if (nbot < asmin) asmin = nbot;
 							if (nbot > asmax) asmax = nbot;
 						}
-
 					}
 				}
 			}
@@ -155,11 +157,17 @@ void rcFilterLedgeSpansImp(rcContext* ctx, const int walkableHeight, const int w
 			// the actual height of the step to climb. Because of that, depending on the number of voxel represented by
 			// the agent radius, the following test might cause more problems than improvements.
 			// That's why we decided to allow to control it via the filterNeighborSlope parameter
-			else if (filterNeighborSlope && (asmax - asmin) > walkableClimb)
-//@UE END
+			else if (neighborSlopeFilterMode == RC_SLOPE_FILTER_RECAST && (asmax - asmin) > walkableClimb)
+
+			{
+				s->data.area = RC_NULL_AREA;
+			}
+			// Compare the step between the min and the max with the maximum step made by a uniform slope for 2 spans (since min and max can be 2 cells apart).
+			else if (neighborSlopeFilterMode == RC_SLOPE_FILTER_USE_HEIGHT_FROM_WALKABLE_SLOPE && (asmax - asmin) > maxStepFor2CellsVx)
 			{
 				s->data.area = RC_NULL_AREA;
 			}
+//@UE END			
 		}
 	}
 }
@@ -174,7 +182,8 @@ void rcFilterLedgeSpansImp(rcContext* ctx, const int walkableHeight, const int w
 /// A span is a ledge if: <tt>rcAbs(currentSpan.smax - neighborSpan.smax) > walkableClimb</tt>
 /// 
 /// @see rcHeightfield, rcConfig
-void rcFilterLedgeSpans(rcContext* ctx, const int walkableHeight, const int walkableClimb, const bool filterNeighborSlope, //UE
+void rcFilterLedgeSpans(rcContext* ctx, const int walkableHeight, const int walkableClimb,
+						const rcNeighborSlopeFilterMode neighborSlopeFilterMode, const rcReal maxStepFromWalkableSlope, const rcReal ch,	//UE
 						rcHeightfield& solid)
 {
 	rcAssert(ctx);
@@ -186,7 +195,7 @@ void rcFilterLedgeSpans(rcContext* ctx, const int walkableHeight, const int walk
 	// Mark border spans.
 	for (int y = 0; y < h; ++y)
 	{
-		rcFilterLedgeSpansImp(ctx, walkableHeight, walkableClimb, filterNeighborSlope, y, solid);	//UE
+		rcFilterLedgeSpansImp(ctx, walkableHeight, walkableClimb, neighborSlopeFilterMode, maxStepFromWalkableSlope, ch, y, solid);	//UE
 	}
 
 	ctx->stopTimer(RC_TIMER_FILTER_BORDER);
@@ -194,7 +203,9 @@ void rcFilterLedgeSpans(rcContext* ctx, const int walkableHeight, const int walk
 
 
 /// @see rcHeightfield, rcConfig
-void rcFilterLedgeSpans(rcContext* ctx, const int walkableHeight, const int walkableClimb, const bool filterNeighborSlope, const int yStart, const int maxYProcess, //UE
+void rcFilterLedgeSpans(rcContext* ctx, const int walkableHeight, const int walkableClimb,
+	const rcNeighborSlopeFilterMode neighborSlopeFilterMode, const rcReal maxStepFromWalkableSlope, const rcReal ch,	// UE
+	const int yStart, const int maxYProcess, //UE
 	rcHeightfield& solid)
 {
 	rcAssert(ctx);
@@ -205,7 +216,7 @@ void rcFilterLedgeSpans(rcContext* ctx, const int walkableHeight, const int walk
 
 	for (int y = yStart; y < h; ++y)
 	{
-		rcFilterLedgeSpansImp(ctx, walkableHeight, walkableClimb, filterNeighborSlope, y, solid);
+		rcFilterLedgeSpansImp(ctx, walkableHeight, walkableClimb, neighborSlopeFilterMode, maxStepFromWalkableSlope, ch, y, solid);
 	}
 
 	ctx->stopTimer(RC_TIMER_FILTER_BORDER);
diff --git a/Engine/Source/Runtime/Navmesh/Public/Recast/Recast.h b/Engine/Source/Runtime/Navmesh/Public/Recast/Recast.h
--- a/Engine/Source/Runtime/Navmesh/Public/Recast/Recast.h
+++ b/Engine/Source/Runtime/Navmesh/Public/Recast/Recast.h
@@ -286,6 +286,9 @@ struct rcConfig
 	/// The distance to erode/shrink the walkable area of the heightfield away from 
 	/// obstructions.  [Limit: >=0] [Units: vx] 
 	int walkableRadius;
+
+	/// Maximum step height in relation to cs and walkableSlopeAngle [Limit: >=0] [Units: wu]
+	rcReal maxStepFromWalkableSlope;
 	
 	/// The maximum allowed length for contour edges along the border of the mesh. [Limit: >=0] [Units: vx] 
 	int maxEdgeLen;
@@ -687,6 +690,14 @@ enum rcRasterizationFlags
 	RC_PROJECT_TO_BOTTOM = 1 << 0,		///< Will create spans from the triangle surface to the bottom of the heightfield
 };
 
+// UE
+enum rcNeighborSlopeFilterMode
+{
+	RC_SLOPE_FILTER_RECAST,							// Use walkableClimb value to filter
+	RC_SLOPE_FILTER_NONE,							// Skip slope filtering
+	RC_SLOPE_FILTER_USE_HEIGHT_FROM_WALKABLE_SLOPE	// Use maximum step height computed from walkableSlopeAngle
+};
+
 /// Applied to the region id field of contour vertices in order to extract the region id.
 /// The region id field of a vertex may have several flags applied to it.  So the
 /// fields value can't be used directly.
@@ -1078,29 +1089,31 @@ NAVMESH_API void rcFilterLowHangingWalkableObstacles(rcContext* ctx, const int w
 
 /// Marks spans that are ledges as not-walkable, by a number of y coords at a time.
 ///  @ingroup recast
-///  @param[in,out]	ctx				The build context to use during the operation.
-///  @param[in]		walkableHeight	Minimum floor to 'ceiling' height that will still allow the floor area to 
-///  								be considered walkable. [Limit: >= 3] [Units: vx]
-///  @param[in]		walkableClimb	Maximum ledge height that is considered to still be traversable. 
-///  								[Limit: >=0] [Units: vx]
-///  @param[in]		filterNeighborSlope If set, a span having neighbors not within the walkable climb range will be marked as non walkable. //UE
-///  @param[in,out]	solid			A fully built heightfield.  (All spans have been added.)
-///  @param[in]		yStart			y coord to start at
-///  @param[in]		maxYProcess	    Max y coords to process (yStart + maxYProcess can be more than solid.height and will be capped to solid.height)
-NAVMESH_API void rcFilterLedgeSpans(rcContext* ctx, const int walkableHeight,
-						const int walkableClimb, const bool filterNeighborSlope, const int yStart, const int maxYProcess, rcHeightfield& solid); //UE
+///  @param[in,out]	ctx							The build context to use during the operation.
+///  @param[in]		walkableHeight				Minimum floor to 'ceiling' height that will still allow the floor area to 
+///  											be considered walkable. [Limit: >= 3] [Units: vx]
+///  @param[in]		walkableClimb				Maximum ledge height that is considered to still be traversable. [Limit: >=0] [Units: vx]
+///  @param[in]		neighborSlopeFilterMode		Change the way neighbors slope filtering is done. //UE
+///  @param[in]		maxStepFromWalkableSlope	Maximum step height in relation to cs and the walkable angle. [Limit: >= 0] [Units: wu] //UE
+///  @param[in]		ch							Cell height. [Limit: >= 0] [Units: wu] //UE
+///  @param[in,out]	solid						A fully built heightfield.  (All spans have been added.)
+///  @param[in]		yStart						y coord to start at
+///  @param[in]		maxYProcess					Max y coords to process (yStart + maxYProcess can be more than solid.height and will be capped to solid.height)
+NAVMESH_API void rcFilterLedgeSpans(rcContext* ctx, const int walkableHeight, const int walkableClimb,
+					const rcNeighborSlopeFilterMode neighborSlopeFilterMode, const rcReal maxStepFromWalkableSlope, const rcReal ch, const int yStart, const int maxYProcess, rcHeightfield& solid); //UE
 
 /// Marks spans that are ledges as not-walkable. 
 ///  @ingroup recast
-///  @param[in,out]	ctx				The build context to use during the operation.
-///  @param[in]		walkableHeight	Minimum floor to 'ceiling' height that will still allow the floor area to 
-///  								be considered walkable. [Limit: >= 3] [Units: vx]
-///  @param[in]		walkableClimb	Maximum ledge height that is considered to still be traversable. 
-///  								[Limit: >=0] [Units: vx]
-///  @param[in]		filterNeighborSlope If set, a span having neighbors not within the walkable climb range will be marked as non walkable. //UE
-///  @param[in,out]	solid			A fully built heightfield.  (All spans have been added.)
-NAVMESH_API void rcFilterLedgeSpans(rcContext* ctx, const int walkableHeight,
-						const int walkableClimb, const bool filterNeighborSlope, rcHeightfield& solid); //UE
+///  @param[in,out]	ctx							The build context to use during the operation.
+///  @param[in]		walkableHeight				Minimum floor to 'ceiling' height that will still allow the floor area to 
+///  											be considered walkable. [Limit: >= 3] [Units: vx]
+///  @param[in]		walkableClimb				Maximum ledge height that is considered to still be traversable. [Limit: >=0] [Units: vx] 
+///  @param[in]		neighborSlopeFilterMode		Change the way neighbors slope filtering is done. //UE
+///  @param[in]		maxStepFromWalkableSlope	Maximum step height in relation to cs and the walkable angle. [Limit: >= 0] [Units: wu] //UE
+///  @param[in]		ch							Cell height. [Limit: >= 0] [Units: wu] //UE
+///  @param[in,out]	solid						A fully built heightfield.  (All spans have been added.)
+NAVMESH_API void rcFilterLedgeSpans(rcContext* ctx, const int walkableHeight, const int walkableClimb,
+					const rcNeighborSlopeFilterMode neighborSlopeFilterMode, const rcReal maxStepFromWalkableSlope, const rcReal ch, rcHeightfield& solid); //UE
 
 /// Marks walkable spans as not walkable if the clearance above the span is less than the specified height. 
 ///  @ingroup recast
