[Navmesh] Addition of a MaxVerticalMergeError and a SimplificationElevationRatio to ARecastNavMesh to allow navmesh to adapt more precisely to vertical collision variances.
Fix navmesh sinking too far below surfaces causing navigation issues.

diff --git a/Engine/Source/Runtime/NavigationSystem/Private/NavMesh/RecastNavMesh.cpp b/Engine/Source/Runtime/NavigationSystem/Private/NavMesh/RecastNavMesh.cpp
--- a/Engine/Source/Runtime/NavigationSystem/Private/NavMesh/RecastNavMesh.cpp
+++ b/Engine/Source/Runtime/NavigationSystem/Private/NavMesh/RecastNavMesh.cpp
@@ -520,7 +520,9 @@ ARecastNavMesh::ARecastNavMesh(const FObjectInitializer& ObjectInitializer)
 	, bDistinctlyDrawTilesBeingBuilt(true)
 	, DrawOffset(10.f)
 	, TilePoolSize(1024)
+	, MaxVerticalMergeError(INT_MAX) // By default, ignore vertical error
 	, MaxSimplificationError(1.3f)	// from RecastDemo
+	, SimplificationElevationRatio(0.f)	// By default, ignore contour simplification from elevation
 	, DefaultMaxSearchNodes(RECAST_MAX_SEARCH_NODES)
 	, DefaultMaxHierarchicalSearchNodes(RECAST_MAX_SEARCH_NODES)
 	, bSortNavigationAreasByCost(true)
diff --git a/Engine/Source/Runtime/NavigationSystem/Private/NavMesh/RecastNavMeshGenerator.cpp b/Engine/Source/Runtime/NavigationSystem/Private/NavMesh/RecastNavMeshGenerator.cpp
--- a/Engine/Source/Runtime/NavigationSystem/Private/NavMesh/RecastNavMeshGenerator.cpp
+++ b/Engine/Source/Runtime/NavigationSystem/Private/NavMesh/RecastNavMeshGenerator.cpp
@@ -3709,12 +3709,12 @@ bool FRecastTileGenerator::GenerateNavigationDataLayer(FNavMeshBuildContext& Bui
 		}
 
 		status = dtBuildTileCacheContours(&GenNavAllocator, *GenerationContext.Layer,
-			TileConfig.walkableClimb, TileConfig.maxSimplificationError, TileConfig.cs, TileConfig.ch,
-			*GenerationContext.ContourSet, *GenerationContext.ClusterSet);
+			TileConfig.walkableClimb, TileConfig.maxVerticalMergeError, TileConfig.maxSimplificationError, TileConfig.simplificationElevationRatio,
+			TileConfig.cs, TileConfig.ch,*GenerationContext.ContourSet, *GenerationContext.ClusterSet);
 #else
 		status = dtBuildTileCacheContours(&GenNavAllocator, *GenerationContext.Layer,
-			TileConfig.walkableClimb, TileConfig.maxSimplificationError, TileConfig.cs, TileConfig.ch,
-			*GenerationContext.ContourSet);
+			TileConfig.walkableClimb, TileConfig.maxVerticalMergeError, TileConfig.maxSimplificationError, TileConfig.simplificationElevationRatio,
+			TileConfig.cs, TileConfig.ch, *GenerationContext.ContourSet);
 #endif //WITH_NAVMESH_CLUSTER_LINKS
 		
 		if (dtStatusFailed(status))
@@ -4701,7 +4701,7 @@ void FRecastNavMeshGenerator::ConfigureBuildProperties(FRecastBuildConfig& OutCo
 
 	OutConfig.borderSize = OutConfig.walkableRadius + 3; // +1 for voxelization rounding, +1 for ledge neighbor access, +1 for occasional errors
 	OutConfig.maxEdgeLen = (int32)(1200.0f / CellSize);
-	OutConfig.maxSimplificationError = 1.3f;
+
 	// hardcoded, but can be overridden by RecastNavMesh params later
 	OutConfig.minRegionArea = (int32)rcSqr(0);
 	OutConfig.mergeRegionArea = (int32)rcSqr(20.f);
@@ -4712,7 +4712,9 @@ void FRecastNavMeshGenerator::ConfigureBuildProperties(FRecastBuildConfig& OutCo
 
 	OutConfig.minRegionArea = (int32)rcSqr(DestNavMesh->MinRegionArea / CellSize);
 	OutConfig.mergeRegionArea = (int32)rcSqr(DestNavMesh->MergeRegionSize / CellSize);
+	OutConfig.maxVerticalMergeError = DestNavMesh->MaxVerticalMergeError;
 	OutConfig.maxSimplificationError = DestNavMesh->MaxSimplificationError;
+	OutConfig.simplificationElevationRatio = DestNavMesh->SimplificationElevationRatio;
 	OutConfig.bPerformVoxelFiltering = DestNavMesh->bPerformVoxelFiltering;
 	OutConfig.bMarkLowHeightAreas = DestNavMesh->bMarkLowHeightAreas;
 	OutConfig.bUseExtraTopCellWhenMarkingAreas = DestNavMesh->bUseExtraTopCellWhenMarkingAreas;
diff --git a/Engine/Source/Runtime/NavigationSystem/Public/NavMesh/RecastNavMesh.h b/Engine/Source/Runtime/NavigationSystem/Public/NavMesh/RecastNavMesh.h
--- a/Engine/Source/Runtime/NavigationSystem/Public/NavMesh/RecastNavMesh.h
+++ b/Engine/Source/Runtime/NavigationSystem/Public/NavMesh/RecastNavMesh.h
@@ -794,10 +794,21 @@ class ARecastNavMesh : public ANavigationData
 	UPROPERTY(EditAnywhere, Category=Generation, config, meta=(ClampMin = "0.0"))
 	float MergeRegionSize;
 
+	/** Maximum vertical deviation between raw contour points to allowing merging (in voxel).
+	 * Use a low value (2-5) depending on CellHeight, AgentMaxStepHeight and AgentMaxSlope, to allow more precise contours (also see SimplificationElevationRatio).
+	 * Use very high value to deactivate (Recast behavior). */
+	UPROPERTY(EditAnywhere, Category = Generation, config, meta = (ClampMin = "0"))
+	int MaxVerticalMergeError;
+	
 	/** How much navigable shapes can get simplified - the higher the value the more freedom */
 	UPROPERTY(EditAnywhere, Category = Generation, config, meta = (ClampMin = "0.0"))
 	float MaxSimplificationError;
 
+	/** When simplifying contours, how much is the vertical error taken into account when comparing with MaxSimplificationError.
+	 * Use 0 to deactivate (Recast behavior), use 1 as a typical value. */
+	UPROPERTY(EditAnywhere, Category = Generation, config, meta = (ClampMin = "0.0"))
+	float SimplificationElevationRatio;
+
 	/** Sets the limit for number of asynchronous tile generators running at one time, also used for some synchronous tasks */
 	UPROPERTY(EditAnywhere, Category = Generation, config, meta = (ClampMin = "0", UIMin = "0"), AdvancedDisplay)
 	int32 MaxSimultaneousTileGenerationJobsCount;
diff --git a/Engine/Source/Runtime/Navmesh/Private/DetourTileCache/DetourTileCache.cpp b/Engine/Source/Runtime/Navmesh/Private/DetourTileCache/DetourTileCache.cpp
--- a/Engine/Source/Runtime/Navmesh/Private/DetourTileCache/DetourTileCache.cpp
+++ b/Engine/Source/Runtime/Navmesh/Private/DetourTileCache/DetourTileCache.cpp
@@ -704,11 +704,11 @@ dtStatus dtTileCache::buildNavMeshTile(const dtCompressedTileRef ref, dtNavMesh*
 	if (!bc.lclusters)
 		return status;
 	status = dtBuildTileCacheContours(m_talloc, *bc.layer, walkableClimbVx,
-									  m_params.maxSimplificationError, m_params.cs, m_params.ch,
+									  m_params.maxVerticalMergeError, m_params.maxSimplificationError, m_params.simplificationElevationRatio, m_params.cs, m_params.ch,
 									  *bc.lcset, *bc.lclusters);
 #else
 	status = dtBuildTileCacheContours(m_talloc, *bc.layer, walkableClimbVx,
-									  m_params.maxSimplificationError, m_params.cs, m_params.ch, *bc.lcset);
+									  m_params.maxVerticalMergeError, m_params.maxSimplificationError, m_params.simplificationElevationRatio, m_params.cs, m_params.ch, *bc.lcset);
 #endif //WITH_NAVMESH_CLUSTER_LINKS
 	//@UE END
 
diff --git a/Engine/Source/Runtime/Navmesh/Private/DetourTileCache/DetourTileCacheBuilder.cpp b/Engine/Source/Runtime/Navmesh/Private/DetourTileCache/DetourTileCacheBuilder.cpp
--- a/Engine/Source/Runtime/Navmesh/Private/DetourTileCache/DetourTileCacheBuilder.cpp
+++ b/Engine/Source/Runtime/Navmesh/Private/DetourTileCache/DetourTileCacheBuilder.cpp
@@ -168,7 +168,7 @@ inline bool overlapRangeExl(const unsigned short amin, const unsigned short amax
 	return (amin >= bmax || amax <= bmin) ? false : true;
 }
 
-static bool appendVertex(dtTempContour& cont, const int x, const int y, const int z, const int r, const unsigned char areaId)
+static bool appendVertex(dtTempContour& cont, const int x, const int y, const int z, const int r, const unsigned char areaId, const int maxVerticalMergeError) // UE
 {
 	// Try to merge with existing segments.
 	if (cont.nverts > 1)
@@ -176,7 +176,7 @@ static bool appendVertex(dtTempContour& cont, const int x, const int y, const in
 		unsigned short* pa = &cont.verts[(cont.nverts-2)*5];
 		unsigned short* pb = &cont.verts[(cont.nverts-1)*5];
 		unsigned short pr = pb[3];
-		if (pr == r)
+		if (pr == r && (dtAbs(pa[1] - y) <= maxVerticalMergeError))	// UE
 		{
 			if (pa[0] == pb[0] && (int)pb[0] == x)
 			{
@@ -270,7 +270,7 @@ static void getNeighbourRegAndArea(dtTileCacheLayer& layer,
 	}
 }
 
-static bool walkContour(dtTileCacheLayer& layer, int x, int y, int idx, unsigned char* flags, dtTempContour& cont)
+static bool walkContour(dtTileCacheLayer& layer, int x, int y, int idx, const int maxVerticalMergeError, unsigned char* flags, dtTempContour& cont) // UE
 {
 	const int w = (int)layer.header->width;
 	const int h = (int)layer.header->height;
@@ -329,7 +329,7 @@ static bool walkContour(dtTileCacheLayer& layer, int x, int y, int idx, unsigned
 			}
 
 			// Try to merge with previous vertex.
-			if (!appendVertex(cont, px, (int)layer.heights[x+y*w], pz, neiReg, neiArea))
+			if (!appendVertex(cont, px, (int)layer.heights[x+y*w], pz, neiReg, neiArea, maxVerticalMergeError)) // UE
 				return false;
 
 			flags[idx] &= ~(1 << dir); // Remove visited edges
@@ -370,7 +370,7 @@ static bool walkContour(dtTileCacheLayer& layer, int x, int y, int idx, unsigned
 
 namespace TileCacheFunc
 {
-	static dtReal distancePtSeg(const int x, const int z, const int px, const int pz, const int qx, const int qz)
+	static dtReal distancePtSegSqr2D(const int x, const int z, const int px, const int pz, const int qx, const int qz) // UE
 	{
 		dtReal pqx = (dtReal)(qx - px);
 		dtReal pqz = (dtReal)(qz - pz);
@@ -392,7 +392,7 @@ namespace TileCacheFunc
 	}
 }
 
-static void simplifyContour(unsigned char area, dtTempContour& cont, const dtReal maxError)
+static void simplifyContour(unsigned char area, dtTempContour& cont, const dtReal maxError, const dtReal elevationRatio, const dtReal cs, const dtReal ch) // UE
 {
 	cont.npoly = 0;
 
@@ -413,6 +413,7 @@ static void simplifyContour(unsigned char area, dtTempContour& cont, const dtRea
 		if (ra != rb || pinnedVertex)
 			cont.poly[cont.npoly++] = (unsigned short)i;
 	}
+
 	if (cont.npoly < 2)
 	{
 		// If there is no transitions at all,
@@ -446,6 +447,8 @@ static void simplifyContour(unsigned char area, dtTempContour& cont, const dtRea
 		cont.poly[cont.npoly++] = (unsigned short)uri;
 	}
 
+	const dtReal heightRatio = elevationRatio * ch / cs; // UE
+
 	// Add points until all raw points are within
 	// error tolerance to the simplified shape.
 	for (int i = 0; i < cont.npoly; )
@@ -454,10 +457,12 @@ static void simplifyContour(unsigned char area, dtTempContour& cont, const dtRea
 
 		const int ai = (int)cont.poly[i];
 		const int ax = (int)cont.verts[ai*5+0];
+		const int ay = (int)cont.verts[ai*5+1]; // UE
 		const int az = (int)cont.verts[ai*5+2];
 
 		const int bi = (int)cont.poly[ii];
 		const int bx = (int)cont.verts[bi*5+0];
+		const int by = (int)cont.verts[bi*5+1]; // UE
 		const int bz = (int)cont.verts[bi*5+2];
 
 		// Find maximum deviation from the segment.
@@ -489,7 +494,22 @@ static void simplifyContour(unsigned char area, dtTempContour& cont, const dtRea
 		{
 			while (ci != endi)
 			{
-				dtReal d = TileCacheFunc::distancePtSeg(cont.verts[ci*5+0], cont.verts[ci*5+2], ax, az, bx, bz);
+//@UE BEGIN
+				dtReal d;
+				if (elevationRatio > 0)
+				{
+					// Instead of multiplying all components by ch or cs to go from voxels to world units, 
+					// we just use the heightRatio (avoiding extra cs multiplication on x and z).
+					const dtReal pt[3] = { (dtReal)cont.verts[ci*5+0], heightRatio*cont.verts[ci*5+1], (dtReal)cont.verts[ci*5+2] };
+					const dtReal a[3] = { (dtReal)ax, heightRatio*ay, (dtReal)az };
+					const dtReal b[3] = { (dtReal)bx, heightRatio*by, (dtReal)bz };
+					d = dtDistancePtSegSqr(pt, a, b);
+				}
+				else
+				{
+					d = TileCacheFunc::distancePtSegSqr2D(cont.verts[ci*5+0], cont.verts[ci*5+2], ax, az, bx, bz);
+				}
+//@UE END
 				if (d > maxd)
 				{
 					maxd = d;
@@ -726,7 +746,7 @@ static void addUniqueRegion(unsigned short* arr, unsigned short v, int& n)
 
 // TODO: move this somewhere else, once the layer meshing is done.
 dtStatus dtBuildTileCacheContours(dtTileCacheAlloc* alloc, dtTileCacheLayer& layer,
-	const int walkableClimb, const dtReal maxError,
+	const int walkableClimb, const int maxVerticalMergeError, const dtReal maxError, const dtReal simplificationElevationRatio, // UE
 	const dtReal cs, const dtReal ch,
 	dtTileCacheContourSet& lcset
 	//@UE BEGIN
@@ -822,14 +842,14 @@ dtStatus dtBuildTileCacheContours(dtTileCacheAlloc* alloc, dtTileCacheLayer& lay
 			if (ri == 0xffff || ri == 0)
 				continue;
 
-			if (!walkContour(layer, x, y, idx, flags, temp))
+			if (!walkContour(layer, x, y, idx, maxVerticalMergeError, flags, temp)) // UE
 			{
 				// Too complex contour.
 				// Note: If you hit here often, try increasing 'maxTempVerts'.
 				return DT_FAILURE | DT_BUFFER_TOO_SMALL;
 			}
 
-			simplifyContour(layer.areas[idx], temp, maxError);
+			simplifyContour(layer.areas[idx], temp, maxError, simplificationElevationRatio, cs, ch); // UE
 
 			// Store contour.
 			if (lcset.nconts >= maxConts)
diff --git a/Engine/Source/Runtime/Navmesh/Public/DetourTileCache/DetourTileCache.h b/Engine/Source/Runtime/Navmesh/Public/DetourTileCache/DetourTileCache.h
--- a/Engine/Source/Runtime/Navmesh/Public/DetourTileCache/DetourTileCache.h
+++ b/Engine/Source/Runtime/Navmesh/Public/DetourTileCache/DetourTileCache.h
@@ -81,7 +81,9 @@ struct dtTileCacheParams
 	dtReal walkableHeight;
 	dtReal walkableRadius;
 	dtReal walkableClimb;
+	dtReal maxVerticalMergeError; // UE
 	dtReal maxSimplificationError;
+	dtReal simplificationElevationRatio; // UE
 	int maxTiles;
 	int maxObstacles;
 //@UE BEGIN
diff --git a/Engine/Source/Runtime/Navmesh/Public/DetourTileCache/DetourTileCacheBuilder.h b/Engine/Source/Runtime/Navmesh/Public/DetourTileCache/DetourTileCacheBuilder.h
--- a/Engine/Source/Runtime/Navmesh/Public/DetourTileCache/DetourTileCacheBuilder.h
+++ b/Engine/Source/Runtime/Navmesh/Public/DetourTileCache/DetourTileCacheBuilder.h
@@ -263,7 +263,7 @@ NAVMESH_API dtStatus dtBuildTileCacheRegionsChunky(dtTileCacheAlloc* alloc,
 
 NAVMESH_API dtStatus dtBuildTileCacheContours(dtTileCacheAlloc* alloc,
 								dtTileCacheLayer& layer,
-								const int walkableClimb, const dtReal maxError,
+								const int walkableClimb, const int maxVerticalMergeError, const dtReal maxError, const dtReal simplificationElevationRatio, // UE
 								const dtReal cs, const dtReal ch,
 								dtTileCacheContourSet& lcset
 								//@UE BEGIN
diff --git a/Engine/Source/Runtime/Navmesh/Public/Recast/Recast.h b/Engine/Source/Runtime/Navmesh/Public/Recast/Recast.h
--- a/Engine/Source/Runtime/Navmesh/Public/Recast/Recast.h
+++ b/Engine/Source/Runtime/Navmesh/Public/Recast/Recast.h
@@ -292,10 +292,19 @@ struct rcConfig
 	
 	/// The maximum allowed length for contour edges along the border of the mesh. [Limit: >=0] [Units: vx] 
 	int maxEdgeLen;
+
+	/// Maximum vertical deviation between raw contour points to allow merging points. [Limit: >=0] [Units: vx]
+	/// Use a low value (2-5) to allow more precise contours (also see SimplificationElevationFactor).
+	/// Use very high value to deactivate (Recast behavior).
+	int maxVerticalMergeError;	// UE
 	
 	/// The maximum distance a simplified contour's border edges should deviate 
 	/// the original raw contour. [Limit: >=0] [Units: wu]
 	rcReal maxSimplificationError;
+
+	/// When simplifying contours, how much is the vertical error taken into account when comparing with MaxSimplificationError. [Limit: >=0]
+	/// Use 0 to deactivate (Recast behavior), use 1 as a typical value.
+	rcReal simplificationElevationRatio;	// UE
 	
 	/// The minimum number of cells allowed to form isolated island areas. [Limit: >=0] [Units: vx] 
 	int minRegionArea;
