Tesselate between regions to improve navmesh height precision.

diff --git a/Engine/Source/Runtime/Navmesh/Private/DetourTileCache/DetourTileCacheBuilder.cpp b/Engine/Source/Runtime/Navmesh/Private/DetourTileCache/DetourTileCacheBuilder.cpp
--- a/Engine/Source/Runtime/Navmesh/Private/DetourTileCache/DetourTileCacheBuilder.cpp
+++ b/Engine/Source/Runtime/Navmesh/Private/DetourTileCache/DetourTileCacheBuilder.cpp
@@ -392,7 +392,7 @@ namespace TileCacheFunc
 	}
 }
 
-static void simplifyContour(unsigned char area, dtTempContour& cont, const dtReal maxError, const dtReal elevationRatio, const dtReal cs, const dtReal ch) // UE
+static void simplifyContour(unsigned char area, unsigned short region, dtTempContour& cont, const dtReal maxError, const dtReal elevationRatio, const dtReal cs, const dtReal ch) // UE
 {
 	cont.npoly = 0;
 
@@ -486,11 +486,12 @@ static void simplifyContour(unsigned char area, dtTempContour& cont, const dtRea
 			endi = ai;
 		}
 
-		// Tessellate only outer edges or edges between areas.
+		// Tessellate only between regions and areas.
 		const unsigned short* ciSrc = &cont.verts[ci*5];
 		const int ciReg = ciSrc[3];
 		const unsigned char ciArea = (unsigned char)ciSrc[4];
-		if (area != ciArea || ciReg == 0xffff)
+		const bool checkRegionChange = elevationRatio > 0;								 // UE
+		if (area != ciArea || ciReg == 0xffff || (checkRegionChange && region != ciReg)) // UE
 		{
 			while (ci != endi)
 			{
@@ -849,7 +850,7 @@ dtStatus dtBuildTileCacheContours(dtTileCacheAlloc* alloc, dtTileCacheLayer& lay
 				return DT_FAILURE | DT_BUFFER_TOO_SMALL;
 			}
 
-			simplifyContour(layer.areas[idx], temp, maxError, simplificationElevationRatio, cs, ch); // UE
+			simplifyContour(layer.areas[idx], ri, temp, maxError, simplificationElevationRatio, cs, ch); // UE
 
 			// Store contour.
 			if (lcset.nconts >= maxConts)
