[Navmesh] Fix crash in DetourTileCacheBuilder canRemoveVertex.

diff --git a/Engine/Source/Runtime/Navmesh/Private/DetourTileCache/DetourTileCacheBuilder.cpp b/Engine/Source/Runtime/Navmesh/Private/DetourTileCache/DetourTileCacheBuilder.cpp
--- a/Engine/Source/Runtime/Navmesh/Private/DetourTileCache/DetourTileCacheBuilder.cpp
+++ b/Engine/Source/Runtime/Navmesh/Private/DetourTileCache/DetourTileCacheBuilder.cpp
@@ -1687,7 +1687,7 @@ static bool canRemoveVertex(dtTileCachePolyMesh& mesh, const unsigned short rem)
 		return false;
 	
 	// Find edges which share the removed vertex.
-	unsigned short edges[MAX_REM_EDGES];
+	unsigned short edges[MAX_REM_EDGES*3];
 	int nedges = 0;
 	
 	for (int i = 0; i < mesh.npolys; ++i)
