[Navmesh] Raw contour improvements.
Allow merging the first contour points with the next and the last to remove differences with contours from other regions.
When walking contours get max of neighbors height.

diff --git a/Engine/Source/Runtime/Navmesh/Private/DetourTileCache/DetourTileCacheBuilder.cpp b/Engine/Source/Runtime/Navmesh/Private/DetourTileCache/DetourTileCacheBuilder.cpp
--- a/Engine/Source/Runtime/Navmesh/Private/DetourTileCache/DetourTileCacheBuilder.cpp
+++ b/Engine/Source/Runtime/Navmesh/Private/DetourTileCache/DetourTileCacheBuilder.cpp
@@ -168,15 +168,16 @@ inline bool overlapRangeExl(const unsigned short amin, const unsigned short amax
 	return (amin >= bmax || amax <= bmin) ? false : true;
 }
 
-static bool appendVertex(dtTempContour& cont, const int x, const int y, const int z, const int r, const unsigned char areaId, const int maxVerticalMergeError) // UE
+static bool appendVertex(dtTempContour& cont, const int x, const int y, const int z, const int neiReg, const unsigned char areaId, const int maxVerticalMergeError) // UE
 {
 	// Try to merge with existing segments.
 	if (cont.nverts > 1)
 	{
+		// pa---------pb---------new(x,y,z)
 		unsigned short* pa = &cont.verts[(cont.nverts-2)*5];
 		unsigned short* pb = &cont.verts[(cont.nverts-1)*5];
 		unsigned short pr = pb[3];
-		if (pr == r && (dtAbs(pa[1] - y) <= maxVerticalMergeError))	// UE
+		if (pr == neiReg && (dtAbs(pa[1] - y) <= maxVerticalMergeError))	// UE
 		{
 			if (pa[0] == pb[0] && (int)pb[0] == x)
 			{
@@ -203,7 +204,7 @@ static bool appendVertex(dtTempContour& cont, const int x, const int y, const in
 	v[0] = (unsigned short)x;
 	v[1] = (unsigned short)y;
 	v[2] = (unsigned short)z;
-	v[3] = (unsigned short)r;
+	v[3] = (unsigned short)neiReg;
 	v[4] = areaId;
 	cont.nverts++;
 
@@ -213,7 +214,7 @@ static bool appendVertex(dtTempContour& cont, const int x, const int y, const in
 
 static void getNeighbourRegAndArea(dtTileCacheLayer& layer,
 	const int ax, const int ay, const int dir,
-	unsigned short& neiReg, unsigned char& neiArea, unsigned char& cornerNeiArea)
+	unsigned short& neiReg, unsigned char& neiArea, unsigned char& cornerNeiArea, unsigned short& neiHeight)	// UE
 {
 	const int w = (int)layer.header->width;
 	const int ia = ax + ay*w;
@@ -237,6 +238,8 @@ static void getNeighbourRegAndArea(dtTileCacheLayer& layer,
 			neiReg = 0xffff;
 			neiArea = 0;
 		}
+
+		neiHeight = layer.heights[ia];	// UE
 	}
 	else
 	{
@@ -246,6 +249,7 @@ static void getNeighbourRegAndArea(dtTileCacheLayer& layer,
 
 		neiReg = layer.regs[ib];
 		neiArea = layer.areas[ib];
+		neiHeight = layer.heights[ib];	 // UE
 
 		// Get area type of the cell diagonal [c] to current cell [a]. Where [b] is direct neighbour in the direction of 'dir'.
 		//   ^
@@ -265,6 +269,7 @@ static void getNeighbourRegAndArea(dtTileCacheLayer& layer,
 			const int cy = by + getDirOffsetY(cdir);
 			const int ic = cx + cy * w;
 			cornerNeiArea = layer.areas[ic];
+			neiHeight = dtMax(neiHeight, layer.heights[ic]);	// UE
 		}
 
 	}
@@ -286,6 +291,7 @@ static bool walkContour(dtTileCacheLayer& layer, int x, int y, int idx, const in
 	unsigned short neiReg = 0xffff;
 	unsigned char neiArea = 0;
 	unsigned char cornerNeiArea = 0;
+	unsigned short neiHeight = 0;	// UE
 	unsigned short prevNeiArea = 0;
 	unsigned short prevCornerNeiArea = 0;
 	bool checkForPinning = false;
@@ -298,7 +304,7 @@ static bool walkContour(dtTileCacheLayer& layer, int x, int y, int idx, const in
 		int ny = y;
 		unsigned char ndir = dir;
 
-		getNeighbourRegAndArea(layer, x, y, dir, neiReg, neiArea, cornerNeiArea);
+		getNeighbourRegAndArea(layer, x, y, dir, neiReg, neiArea, cornerNeiArea, neiHeight);	// UE
 
 		if (neiReg != layer.regs[x+y*w])
 		{
@@ -329,7 +335,8 @@ static bool walkContour(dtTileCacheLayer& layer, int x, int y, int idx, const in
 			}
 
 			// Try to merge with previous vertex.
-			if (!appendVertex(cont, px, (int)layer.heights[x+y*w], pz, neiReg, neiArea, maxVerticalMergeError)) // UE
+			const int py = dtMax(neiHeight, (int)layer.heights[x+y*w]);	// UE
+			if (!appendVertex(cont, px, py, pz, neiReg, neiArea, maxVerticalMergeError)) // UE
 				return false;
 
 			flags[idx] &= ~(1 << dir); // Remove visited edges
@@ -365,6 +372,39 @@ static bool walkContour(dtTileCacheLayer& layer, int x, int y, int idx, const in
 	if (pa[0] == pb[0] && pa[2] == pb[2])
 		cont.nverts--;
 
+//@UE BEGIN
+	// Check if first vertex should be merged.
+	if (cont.nverts > 1)
+	{
+		unsigned short* last = &cont.verts[(cont.nverts-1)*5];
+		unsigned short* first = &cont.verts[0*5];
+		unsigned short* next = &cont.verts[1*5];
+
+		// Check if we can remove first vertex. First vertex will become last vertex.
+		if (first[3] == next[3] && (dtAbs(next[1] - last[1]) <= maxVerticalMergeError))
+		{
+			if (last[0] == first[0] && first[0] == next[0])
+			{
+				// The verts are aligned aling x-axis, update z.
+				first[1] = last[1];
+				first[2] = last[2];
+				first[3] = last[3];
+				first[4] = last[4]; 
+				cont.nverts--;	// remove last
+			}
+			else if (last[2] == first[2] && first[2] == next[2])
+			{
+				// The verts are aligned aling z-axis, update x.
+				first[0] = last[0];
+				first[1] = last[1];
+				first[3] = last[3];
+				first[4] = last[4];
+				cont.nverts--; // remove last
+			}
+		}
+	}	
+//@UE END
+	
 	return true;
 }	
 
