Fix comment in dtNavMesh

diff --git a/Engine/Source/Runtime/Navmesh/Public/Detour/DetourNavMesh.h b/Engine/Source/Runtime/Navmesh/Public/Detour/DetourNavMesh.h
--- a/Engine/Source/Runtime/Navmesh/Public/Detour/DetourNavMesh.h
+++ b/Engine/Source/Runtime/Navmesh/Public/Detour/DetourNavMesh.h
@@ -886,7 +886,7 @@ private:
 	/// Returns pointer to tile in the tile array.
 	NAVMESH_API dtMeshTile* getTile(int i);
 
-	/// Returns neighbour tile based on side.
+	/// Returns the tile at the coordinates.
 	NAVMESH_API int getTilesAt(const int x, const int y,
 				   dtMeshTile** tiles, const int maxTiles) const;
 
