[Navmesh] Fix debug display issue in duDebugDrawHeightfieldLayer

diff --git a/Engine/Source/Runtime/Navmesh/Private/DebugUtils/RecastDebugDraw.cpp b/Engine/Source/Runtime/Navmesh/Private/DebugUtils/RecastDebugDraw.cpp
--- a/Engine/Source/Runtime/Navmesh/Private/DebugUtils/RecastDebugDraw.cpp
+++ b/Engine/Source/Runtime/Navmesh/Private/DebugUtils/RecastDebugDraw.cpp
@@ -424,7 +424,7 @@ void duDebugDrawHeightfieldLayer(duDebugDraw* dd, const struct rcHeightfieldLaye
 		{
 			const int lidx = x+y*w;
 			const int lh = (int)layer.heights[lidx];
-			if (h == 0xffff) continue;
+			if (lh == 0xffff) continue;
 			const unsigned char area = layer.areas[lidx];
 			
 			unsigned int col;
