Fix navmesh generation issue where polygon would wrongly connect to a different height, resulting in a broken navmesh surface.

diff --git a/Engine/Source/Runtime/Navmesh/Private/DetourTileCache/DetourTileCacheDetail.cpp b/Engine/Source/Runtime/Navmesh/Private/DetourTileCache/DetourTileCacheDetail.cpp
--- a/Engine/Source/Runtime/Navmesh/Private/DetourTileCache/DetourTileCacheDetail.cpp
+++ b/Engine/Source/Runtime/Navmesh/Private/DetourTileCache/DetourTileCacheDetail.cpp
@@ -38,6 +38,11 @@ struct dtHeightPatch
 	int xmin, ymin, width, height;
 };
 
+inline bool isCellConnected(const dtTileCacheLayer& layer, const int idx, const int dir)
+{
+	return (layer.cons[idx] & (1 << dir)) != 0;
+}
+
 static void getLayerHeightData(dtTileCacheLayer& layer,
 	const unsigned short* poly, const unsigned short* verts, const int nverts,
 	dtHeightPatch& hp, dtIntArray& stack)
@@ -131,7 +136,8 @@ static void getLayerHeightData(dtTileCacheLayer& layer,
 			if (ax < hp.xmin || ax >= (hp.xmin+hp.width) ||
 				ay < hp.ymin || ay >= (hp.ymin+hp.height) ||
 				layer.heights[hidx] == DT_UNSET_LAYER_HEIGHT ||
-				hp.data[idx] != 0)
+				hp.data[idx] != 0 ||
+				!isCellConnected(layer, cx+cy*layer.header->width, dir))
 			{
 				continue;
 			}
@@ -180,7 +186,8 @@ static void getLayerHeightData(dtTileCacheLayer& layer,
 			if (ax < hp.xmin || ax >= (hp.xmin+hp.width) ||
 				ay < hp.ymin || ay >= (hp.ymin+hp.height) ||
 				hp.data[idx] != DT_UNSET_PATCH_HEIGHT ||
-				layer.heights[hidx] == DT_UNSET_LAYER_HEIGHT)
+				layer.heights[hidx] == DT_UNSET_LAYER_HEIGHT ||
+				!isCellConnected(layer, cx+cy*layer.header->width, dir))
 			{
 				continue;
 			}
