[NavMesh] Removed MaxVerticalMergeError and instead prevent point merging when adding raw contour points if simplificationElevationRatio is being used.

diff --git a/Engine/Source/Runtime/NavigationSystem/Private/NavMesh/RecastNavMeshGenerator.cpp b/Engine/Source/Runtime/NavigationSystem/Private/NavMesh/RecastNavMeshGenerator.cpp
--- a/Engine/Source/Runtime/NavigationSystem/Private/NavMesh/RecastNavMeshGenerator.cpp
+++ b/Engine/Source/Runtime/NavigationSystem/Private/NavMesh/RecastNavMeshGenerator.cpp
@@ -3728,11 +3728,11 @@ bool FRecastTileGenerator::GenerateNavigationDataLayer(FNavMeshBuildContext& Bui
 		}
 
 		status = dtBuildTileCacheContours(&GenNavAllocator, *GenerationContext.Layer,
-			TileConfig.walkableClimb, TileConfig.maxVerticalMergeError, TileConfig.maxSimplificationError, TileConfig.simplificationElevationRatio,
+			TileConfig.walkableClimb, TileConfig.maxSimplificationError, TileConfig.simplificationElevationRatio,
 			TileConfig.cs, TileConfig.ch,*GenerationContext.ContourSet, *GenerationContext.ClusterSet, bSkipContourSimplification);
 #else
 		status = dtBuildTileCacheContours(&GenNavAllocator, *GenerationContext.Layer,
-			TileConfig.walkableClimb, TileConfig.maxVerticalMergeError, TileConfig.maxSimplificationError, TileConfig.simplificationElevationRatio,
+			TileConfig.walkableClimb, TileConfig.maxSimplificationError, TileConfig.simplificationElevationRatio,
 			TileConfig.cs, TileConfig.ch, *GenerationContext.ContourSet, bSkipContourSimplification);
 #endif //WITH_NAVMESH_CLUSTER_LINKS
 		
@@ -4736,7 +4736,6 @@ void FRecastNavMeshGenerator::ConfigureBuildProperties(FRecastBuildConfig& OutCo
 
 	OutConfig.minRegionArea = (int32)rcSqr(DestNavMesh->MinRegionArea / CellSize);
 	OutConfig.mergeRegionArea = (int32)rcSqr(DestNavMesh->MergeRegionSize / CellSize);
-	OutConfig.maxVerticalMergeError = DestNavMesh->MaxVerticalMergeError;
 	OutConfig.maxSimplificationError = DestNavMesh->MaxSimplificationError;
 	OutConfig.simplificationElevationRatio = DestNavMesh->SimplificationElevationRatio;
 	OutConfig.bPerformVoxelFiltering = DestNavMesh->bPerformVoxelFiltering;
diff --git a/Engine/Source/Runtime/NavigationSystem/Public/NavMesh/RecastNavMesh.h b/Engine/Source/Runtime/NavigationSystem/Public/NavMesh/RecastNavMesh.h
--- a/Engine/Source/Runtime/NavigationSystem/Public/NavMesh/RecastNavMesh.h
+++ b/Engine/Source/Runtime/NavigationSystem/Public/NavMesh/RecastNavMesh.h
@@ -805,7 +805,8 @@ class ARecastNavMesh : public ANavigationData
 	/** Maximum vertical deviation between raw contour points to allowing merging (in voxel).
 	 * Use a low value (2-5) depending on CellHeight, AgentMaxStepHeight and AgentMaxSlope, to allow more precise contours (also see SimplificationElevationRatio).
 	 * Use very high value to deactivate (Recast behavior). */
-	UPROPERTY(EditAnywhere, Category = Generation, config, meta = (ClampMin = "0"))
+	UE_DEPRECATED(5.5, "Not used anymore, the behavior is now binded to SimplificationElevationRatio.")
+	UPROPERTY(config)
 	int MaxVerticalMergeError;
 	
 	/** How much navigable shapes can get simplified - the higher the value the more freedom */
diff --git a/Engine/Source/Runtime/Navmesh/Private/DetourTileCache/DetourTileCache.cpp b/Engine/Source/Runtime/Navmesh/Private/DetourTileCache/DetourTileCache.cpp
--- a/Engine/Source/Runtime/Navmesh/Private/DetourTileCache/DetourTileCache.cpp
+++ b/Engine/Source/Runtime/Navmesh/Private/DetourTileCache/DetourTileCache.cpp
@@ -704,11 +704,11 @@ dtStatus dtTileCache::buildNavMeshTile(const dtCompressedTileRef ref, dtNavMesh*
 	if (!bc.lclusters)
 		return status;
 	status = dtBuildTileCacheContours(m_talloc, *bc.layer, walkableClimbVx,
-									  m_params.maxVerticalMergeError, m_params.maxSimplificationError, m_params.simplificationElevationRatio, m_params.cs, m_params.ch,
+									  m_params.maxSimplificationError, m_params.simplificationElevationRatio, m_params.cs, m_params.ch,
 									  *bc.lcset, *bc.lclusters);
 #else
 	status = dtBuildTileCacheContours(m_talloc, *bc.layer, walkableClimbVx,
-									  m_params.maxVerticalMergeError, m_params.maxSimplificationError, m_params.simplificationElevationRatio, m_params.cs, m_params.ch, *bc.lcset);
+									  m_params.maxSimplificationError, m_params.simplificationElevationRatio, m_params.cs, m_params.ch, *bc.lcset);
 #endif //WITH_NAVMESH_CLUSTER_LINKS
 	//@UE END
 
diff --git a/Engine/Source/Runtime/Navmesh/Private/DetourTileCache/DetourTileCacheBuilder.cpp b/Engine/Source/Runtime/Navmesh/Private/DetourTileCache/DetourTileCacheBuilder.cpp
--- a/Engine/Source/Runtime/Navmesh/Private/DetourTileCache/DetourTileCacheBuilder.cpp
+++ b/Engine/Source/Runtime/Navmesh/Private/DetourTileCache/DetourTileCacheBuilder.cpp
@@ -152,7 +152,7 @@ struct dtTempContour
 		poly(pbuf), npoly(0), cpoly(npbuf) 
 	{
 	}
-	unsigned short* verts;
+	unsigned short* verts;	// x, y, z, reg, area
 	int nverts;
 	int cverts;
 	unsigned short* poly;
@@ -169,16 +169,17 @@ inline bool overlapRangeExl(const unsigned short amin, const unsigned short amax
 	return (amin >= bmax || amax <= bmin) ? false : true;
 }
 
-static bool appendVertex(dtTempContour& cont, const int x, const int y, const int z, const int neiReg, const unsigned char areaId, const int maxVerticalMergeError) // UE
+// Returns true on success, false if there was an error adding a vertex.
+static bool appendVertex(dtTempContour& cont, const int x, const int y, const int z, const int neiReg, const unsigned char areaId, const bool allowMerging) // UE
 {
 	// Try to merge with existing segments.
 	if (cont.nverts > 1)
 	{
 		// pa---------pb---------new(x,y,z)
-		unsigned short* pa = &cont.verts[(cont.nverts-2)*5];
+		const unsigned short* pa = &cont.verts[(cont.nverts-2)*5];
 		unsigned short* pb = &cont.verts[(cont.nverts-1)*5];
-		unsigned short pr = pb[3];
-		if (pr == neiReg && (dtAbs(pa[1] - y) <= maxVerticalMergeError))	// UE
+		const unsigned short pr = pb[3];
+		if (allowMerging && pr == neiReg)	// UE
 		{
 			if (pa[0] == pb[0] && (int)pb[0] == x)
 			{
@@ -276,7 +277,7 @@ static void getNeighbourRegAndArea(dtTileCacheLayer& layer,
 	}
 }
 
-static bool walkContour(dtTileCacheLayer& layer, int x, int y, int idx, const int maxVerticalMergeError, unsigned char* flags, dtTempContour& cont, int& contourIndex) // UE
+static bool walkContour(dtTileCacheLayer& layer, int x, int y, int idx, const bool allowMerging, unsigned char* flags, dtTempContour& cont, int& contourIndex) // UE
 {
 	const int w = (int)layer.header->width;
 	const int h = (int)layer.header->height;
@@ -313,7 +314,7 @@ static bool walkContour(dtTileCacheLayer& layer, int x, int y, int idx, const in
 			if (checkForPinning)
 			{
 				// Detect if there was 8-connected area type change during the turn.
-				// If it did, we need to make sure the vertex adde during turn does not get simplified.
+				// If it did, we need to make sure the vertex added during turn does not get simplified.
 				// AB
 				// xC
 				// x = current location, A = prevNeiArea, B = prevCornerNeiArea, C = neiArea.
@@ -337,7 +338,7 @@ static bool walkContour(dtTileCacheLayer& layer, int x, int y, int idx, const in
 
 			// Try to merge with previous vertex.
 			const int py = dtMax(neiHeight, (int)layer.heights[x+y*w]);	// UE
-			if (!appendVertex(cont, px, py, pz, neiReg, neiArea, maxVerticalMergeError)) // UE
+			if (!appendVertex(cont, px, py, pz, neiReg, neiArea, allowMerging)) // UE
 				return false;
 
 			flags[idx] &= ~(1 << dir); // Remove visited edges
@@ -382,7 +383,7 @@ static bool walkContour(dtTileCacheLayer& layer, int x, int y, int idx, const in
 		unsigned short* next = &cont.verts[1*5];
 
 		// Check if we can remove first vertex. First vertex will become last vertex.
-		if (first[3] == next[3] && (dtAbs(next[1] - last[1]) <= maxVerticalMergeError))
+		if (allowMerging && first[3] == next[3])
 		{
 			if (last[0] == first[0] && first[0] == next[0])
 			{
@@ -618,7 +619,7 @@ static unsigned short getCornerHeight(dtTileCacheLayer& layer, const int x, cons
 
 	unsigned char portal = 0xf;
 	unsigned short height = 0;
-	unsigned short preg = 0xffff;
+	unsigned short preg = 0xffff;	// portal region
 	bool allSameReg = true;
 
 	for (int dz = -1; dz <= 0; ++dz)
@@ -646,8 +647,10 @@ static unsigned short getCornerHeight(dtTileCacheLayer& layer, const int x, cons
 
 	int portalCount = 0;
 	for (int dir = 0; dir < 4; ++dir)
+	{
 		if (portal & (1<<dir))
 			portalCount++;
+	}
 
 	shouldRemove = false;
 	if (n > 1 && portalCount == 1 && allSameReg)
@@ -790,7 +793,7 @@ static void addUniqueRegion(unsigned short* arr, unsigned short v, int& n)
 
 // TODO: move this somewhere else, once the layer meshing is done.
 dtStatus dtBuildTileCacheContours(dtTileCacheAlloc* alloc, dtTileCacheLayer& layer,
-	const int walkableClimb, const int maxVerticalMergeError, const dtReal maxError, const dtReal simplificationElevationRatio, // UE
+	const int walkableClimb, const dtReal maxError, const dtReal simplificationElevationRatio, // UE
 	const dtReal cs, const dtReal ch,
 	dtTileCacheContourSet& lcset
 	//@UE BEGIN
@@ -872,6 +875,9 @@ dtStatus dtBuildTileCacheContours(dtTileCacheAlloc* alloc, dtTileCacheLayer& lay
 	dtIntArray nlinks(maxConts);
 	dtIntArray linksBase(maxConts);
 
+	// Only allow merging when not using height in the contour simplification process.
+	const bool allowMerging = (simplificationElevationRatio == 0);
+	
 	// Find contours.
 	int contourIndex = 0;	// UE
 	for (int y = 0; y < h; ++y)
@@ -888,7 +894,7 @@ dtStatus dtBuildTileCacheContours(dtTileCacheAlloc* alloc, dtTileCacheLayer& lay
 			if (ri == 0xffff || ri == 0)
 				continue;
 
-			if (!walkContour(layer, x, y, idx, maxVerticalMergeError, flags, temp, contourIndex)) // UE
+			if (!walkContour(layer, x, y, idx, allowMerging, flags, temp, contourIndex)) // UE
 			{
 				// Too complex contour.
 				// Note: If you hit here often, try increasing 'maxTempVerts'.
@@ -960,7 +966,7 @@ dtStatus dtBuildTileCacheContours(dtTileCacheAlloc* alloc, dtTileCacheLayer& lay
 					dst[1] = lh;
 					dst[2] = v[2];
 
-					// Store portal direction and remove status to the fourth component.
+					// Store portal direction and store remove status to the fourth component.
 					dst[3] = 0x0f;
 					if (nei != 0xffff && nei >= 0xf800)
 						dst[3] = (unsigned char)(nei - 0xf800);
@@ -1189,15 +1195,18 @@ static bool buildMeshAdjacency(dtTileCacheAlloc* alloc,
 	
 	for (int i = 0; i < nverts; i++)
 		firstEdge[i] = DT_TILECACHE_NULL_IDX;
-	
+
+	// Add edges
 	for (int i = 0; i < npolys; ++i)
 	{
-		unsigned short* t = &polys[i*MAX_VERTS_PER_POLY*2];
+		const unsigned short* t = &polys[i*MAX_VERTS_PER_POLY*2];
 		for (int j = 0; j < MAX_VERTS_PER_POLY; ++j)
 		{
-			if (t[j] == DT_TILECACHE_NULL_IDX) break;
-			unsigned short v0 = t[j];
-			unsigned short v1 = (j+1 >= MAX_VERTS_PER_POLY || t[j+1] == DT_TILECACHE_NULL_IDX) ? t[0] : t[j+1];
+			if (t[j] == DT_TILECACHE_NULL_IDX)
+				break;
+
+			const unsigned short v0 = t[j];
+			const unsigned short v1 = (j+1 >= MAX_VERTS_PER_POLY || t[j+1] == DT_TILECACHE_NULL_IDX) ? t[0] : t[j+1];
 			if (v0 < v1)
 			{
 				TileCacheData::rcEdge& edge = edges[edgeCount];
@@ -1217,10 +1226,12 @@ static bool buildMeshAdjacency(dtTileCacheAlloc* alloc,
 	
 	for (int i = 0; i < npolys; ++i)
 	{
-		unsigned short* t = &polys[i*MAX_VERTS_PER_POLY*2];
+		const unsigned short* t = &polys[i*MAX_VERTS_PER_POLY*2];
 		for (int j = 0; j < MAX_VERTS_PER_POLY; ++j)
 		{
-			if (t[j] == DT_TILECACHE_NULL_IDX) break;
+			if (t[j] == DT_TILECACHE_NULL_IDX)
+				break;
+
 			unsigned short v0 = t[j];
 			unsigned short v1 = (j+1 >= MAX_VERTS_PER_POLY || t[j+1] == DT_TILECACHE_NULL_IDX) ? t[0] : t[j+1];
 			if (v0 > v1)
@@ -1259,7 +1270,7 @@ static bool buildMeshAdjacency(dtTileCacheAlloc* alloc,
 	// Mark portal edges.
 	for (int i = 0; i < lcset.nconts; ++i)
 	{
-		dtTileCacheContour& cont = lcset.conts[i];
+		const dtTileCacheContour& cont = lcset.conts[i];
 		if (cont.nverts < 3)
 			continue;
 		
@@ -1347,7 +1358,7 @@ static bool buildMeshAdjacency(dtTileCacheAlloc* alloc,
 			p0[MAX_VERTS_PER_POLY + e.polyEdge[0]] = e.poly[1];
 			p1[MAX_VERTS_PER_POLY + e.polyEdge[1]] = e.poly[0];
 		}
-		else if (e.polyEdge[1] != 0xff)
+		else if (e.polyEdge[1] != 0xff) // if we have a direction
 		{
 			unsigned short* p0 = &polys[e.poly[0]*MAX_VERTS_PER_POLY*2];
 			p0[MAX_VERTS_PER_POLY + e.polyEdge[0]] = 0x8000 | (unsigned short)e.polyEdge[1];
@@ -1702,7 +1713,7 @@ static bool canRemoveVertex(dtTileCachePolyMesh& mesh, const unsigned short rem)
 	int numRemainingEdges = 0;
 	for (int i = 0; i < mesh.npolys; ++i)
 	{
-		unsigned short* p = &mesh.polys[i*MAX_VERTS_PER_POLY*2];
+		const unsigned short* p = &mesh.polys[i*MAX_VERTS_PER_POLY*2];
 		const int nv = countPolyVerts(p);
 		int numRemoved = 0;
 		int numVerts = 0;
@@ -1740,7 +1751,7 @@ static bool canRemoveVertex(dtTileCachePolyMesh& mesh, const unsigned short rem)
 	
 	for (int i = 0; i < mesh.npolys; ++i)
 	{
-		unsigned short* p = &mesh.polys[i*MAX_VERTS_PER_POLY*2];
+		const unsigned short* p = &mesh.polys[i*MAX_VERTS_PER_POLY*2];
 		const int nv = countPolyVerts(p);
 		
 		// Collect edges which touches the removed vertex.
@@ -1799,7 +1810,7 @@ static dtStatus removeVertex(dtTileCacheLogContext* ctx, dtTileCachePolyMesh& me
 	int numRemovedVerts = 0;
 	for (int i = 0; i < mesh.npolys; ++i)
 	{
-		unsigned short* p = &mesh.polys[i*MAX_VERTS_PER_POLY*2];
+		const unsigned short* p = &mesh.polys[i*MAX_VERTS_PER_POLY*2];
 		const int nv = countPolyVerts(p);
 		for (int j = 0; j < nv; ++j)
 		{
@@ -1847,7 +1858,7 @@ static dtStatus removeVertex(dtTileCacheLogContext* ctx, dtTileCachePolyMesh& me
 				}
 			}
 			// Remove the polygon.
-			unsigned short* p2 = &mesh.polys[(mesh.npolys-1)*MAX_VERTS_PER_POLY*2];
+			const unsigned short* p2 = &mesh.polys[(mesh.npolys-1)*MAX_VERTS_PER_POLY*2];
 			memcpy(p,p2,sizeof(unsigned short)*MAX_VERTS_PER_POLY);
 			memset(p+MAX_VERTS_PER_POLY,0xff,sizeof(unsigned short)*MAX_VERTS_PER_POLY);
 			mesh.areas[i] = mesh.areas[mesh.npolys-1];
@@ -2041,10 +2052,10 @@ static dtStatus removeVertex(dtTileCacheLogContext* ctx, dtTileCachePolyMesh& me
 	for (int i = 0; i < npolys; ++i)
 	{
 		if (mesh.npolys >= maxTris) break;
-		unsigned short* p = &mesh.polys[mesh.npolys*MAX_VERTS_PER_POLY*2];
-		memset(p,0xff,sizeof(unsigned short)*MAX_VERTS_PER_POLY*2);
+		unsigned short* newPoly = &mesh.polys[mesh.npolys*MAX_VERTS_PER_POLY*2];
+		memset(newPoly,0xff,sizeof(unsigned short)*MAX_VERTS_PER_POLY*2);
 		for (int j = 0; j < MAX_VERTS_PER_POLY; ++j)
-			p[j] = polys[i*MAX_VERTS_PER_POLY+j];
+			newPoly[j] = polys[i*MAX_VERTS_PER_POLY+j];
 		mesh.areas[mesh.npolys] = pareas[i];
 		mesh.npolys++;
 		if (mesh.npolys > maxTris)
@@ -2248,10 +2259,10 @@ dtStatus dtBuildTileCachePolyMesh(dtTileCacheAlloc* alloc,
 		// Store polygons.
 		for (int j = 0; j < npolys; ++j)
 		{
-			unsigned short* p = &mesh.polys[mesh.npolys*MAX_VERTS_PER_POLY*2];
-			unsigned short* q = &polys[j*MAX_VERTS_PER_POLY];
+			unsigned short* newPoly = &mesh.polys[mesh.npolys*MAX_VERTS_PER_POLY*2];
+			const unsigned short* q = &polys[j*MAX_VERTS_PER_POLY];
 			for (int k = 0; k < MAX_VERTS_PER_POLY; ++k)
-				p[k] = q[k];
+				newPoly[k] = q[k];
 			mesh.areas[mesh.npolys] = cont.area;
 			mesh.regs[mesh.npolys] = cont.reg;
 			mesh.npolys++;
diff --git a/Engine/Source/Runtime/Navmesh/Public/DetourTileCache/DetourTileCache.h b/Engine/Source/Runtime/Navmesh/Public/DetourTileCache/DetourTileCache.h
--- a/Engine/Source/Runtime/Navmesh/Public/DetourTileCache/DetourTileCache.h
+++ b/Engine/Source/Runtime/Navmesh/Public/DetourTileCache/DetourTileCache.h
@@ -81,7 +81,6 @@ struct dtTileCacheParams
 	dtReal walkableHeight;
 	dtReal walkableRadius;
 	dtReal walkableClimb;
-	dtReal maxVerticalMergeError; // UE
 	dtReal maxSimplificationError;
 	dtReal simplificationElevationRatio; // UE
 	int maxTiles;
diff --git a/Engine/Source/Runtime/Navmesh/Public/DetourTileCache/DetourTileCacheBuilder.h b/Engine/Source/Runtime/Navmesh/Public/DetourTileCache/DetourTileCacheBuilder.h
--- a/Engine/Source/Runtime/Navmesh/Public/DetourTileCache/DetourTileCacheBuilder.h
+++ b/Engine/Source/Runtime/Navmesh/Public/DetourTileCache/DetourTileCacheBuilder.h
@@ -263,7 +263,7 @@ NAVMESH_API dtStatus dtBuildTileCacheRegionsChunky(dtTileCacheAlloc* alloc,
 
 NAVMESH_API dtStatus dtBuildTileCacheContours(dtTileCacheAlloc* alloc,
 								dtTileCacheLayer& layer,
-								const int walkableClimb, const int maxVerticalMergeError, const dtReal maxError, const dtReal simplificationElevationRatio, // UE
+								const int walkableClimb, const dtReal maxError, const dtReal simplificationElevationRatio, // UE
 								const dtReal cs, const dtReal ch,
 								dtTileCacheContourSet& lcset
 								//@UE BEGIN
diff --git a/Engine/Source/Runtime/Navmesh/Public/Recast/Recast.h b/Engine/Source/Runtime/Navmesh/Public/Recast/Recast.h
--- a/Engine/Source/Runtime/Navmesh/Public/Recast/Recast.h
+++ b/Engine/Source/Runtime/Navmesh/Public/Recast/Recast.h
@@ -292,11 +292,6 @@ struct rcConfig
 	
 	/// The maximum allowed length for contour edges along the border of the mesh. [Limit: >=0] [Units: vx] 
 	int maxEdgeLen;
-
-	/// Maximum vertical deviation between raw contour points to allow merging points. [Limit: >=0] [Units: vx]
-	/// Use a low value (2-5) to allow more precise contours (also see SimplificationElevationFactor).
-	/// Use very high value to deactivate (Recast behavior).
-	int maxVerticalMergeError;	// UE
 	
 	/// The maximum distance a simplified contour's border edges should deviate 
 	/// the original raw contour. [Limit: >=0] [Units: wu]
