[NavMesh] Fix contour simplification to be symetrical when using SimplificationElevationRatio.

diff --git a/Engine/Source/Runtime/Navmesh/Private/DetourTileCache/DetourTileCacheBuilder.cpp b/Engine/Source/Runtime/Navmesh/Private/DetourTileCache/DetourTileCacheBuilder.cpp
--- a/Engine/Source/Runtime/Navmesh/Private/DetourTileCache/DetourTileCacheBuilder.cpp
+++ b/Engine/Source/Runtime/Navmesh/Private/DetourTileCache/DetourTileCacheBuilder.cpp
@@ -500,14 +500,14 @@ static void simplifyContour(unsigned char area, unsigned short region, dtTempCon
 		int ii = (i+1) % cont.npoly;
 
 		const int ai = (int)cont.poly[i];
-		const int ax = (int)cont.verts[ai*5+0];
-		const int ay = (int)cont.verts[ai*5+1]; // UE
-		const int az = (int)cont.verts[ai*5+2];
+		int ax = (int)cont.verts[ai*5+0];
+		int ay = (int)cont.verts[ai*5+1]; // UE
+		int az = (int)cont.verts[ai*5+2];
 
 		const int bi = (int)cont.poly[ii];
-		const int bx = (int)cont.verts[bi*5+0];
-		const int by = (int)cont.verts[bi*5+1]; // UE
-		const int bz = (int)cont.verts[bi*5+2];
+		int bx = (int)cont.verts[bi*5+0];
+		int by = (int)cont.verts[bi*5+1]; // UE
+		int bz = (int)cont.verts[bi*5+2];
 
 		// Find maximum deviation from the segment.
 		dtReal maxd = 0;
@@ -528,6 +528,12 @@ static void simplifyContour(unsigned char area, unsigned short region, dtTempCon
 			cinc = cont.nverts-1;
 			ci = (bi+cinc) % cont.nverts;
 			endi = ai;
+
+			// Because of floating point imprecision, dtDistancePtSegSqr to a-b might be slightly differ from dtDistancePtSegSqr to b-a.
+			// Swap points because we need the maximum deviation to be computed the same way for opposite segments to match.
+			dtSwap(ax, bx);	// UE
+			dtSwap(ay, by);	// UE
+			dtSwap(az, bz);	// UE
 		}
 
 		// Tessellate only between regions and areas.
