[NavMesh] While adding raw contour points, use the height from the highest connected corner cell.

diff --git a/Engine/Source/Runtime/Navmesh/Private/DetourTileCache/DetourTileCacheBuilder.cpp b/Engine/Source/Runtime/Navmesh/Private/DetourTileCache/DetourTileCacheBuilder.cpp
--- a/Engine/Source/Runtime/Navmesh/Private/DetourTileCache/DetourTileCacheBuilder.cpp
+++ b/Engine/Source/Runtime/Navmesh/Private/DetourTileCache/DetourTileCacheBuilder.cpp
@@ -213,26 +213,33 @@ static bool appendVertex(dtTempContour& cont, const int x, const int y, const in
 	return true;
 }
 
-
-static void getNeighbourRegAndArea(dtTileCacheLayer& layer,
-	const int ax, const int ay, const int dir,
-	unsigned short& neiReg, unsigned char& neiArea, unsigned char& cornerNeiArea, unsigned short& neiHeight)	// UE
+//@UE BEGIN
+static void getNeighbourRegAndAreaAndVertexHeight(dtTileCacheLayer& layer,
+	const int ax, const int ay, const int abDir,
+	unsigned short& neiReg, unsigned char& neiArea, unsigned char& cornerNeiArea, unsigned short& vertexHeight)
 {
+	// [a] is the current cell, [b] is the direct neighbour in the direction 'dir'.
+	//   ^
+	//  [b][c]
+	//  [a][d]
+	
 	const int w = (int)layer.header->width;
 	const int ia = ax + ay*w;
-
-	const unsigned char con = layer.cons[ia] & 0xf;
+	const unsigned char acon = layer.cons[ia] & 0xf;
 	const unsigned char portal = layer.cons[ia] >> 4;
-	const unsigned char mask = (unsigned char)(1<<dir);
+	const unsigned char abDirMask = (unsigned char)(1 << abDir);
 
+	const int bcDir = (abDir + 1) & 0x3;
+	const unsigned char bcDirMask = (unsigned char)(1 << bcDir);
+	
 	cornerNeiArea = 0;
 
-	if ((con & mask) == 0)
+	if ((acon & abDirMask) == 0)
 	{
 		// No connection, return portal or hard edge.
-		if (portal & mask)
+		if (portal & abDirMask)
 		{
-			neiReg = 0xf800 + (unsigned char)dir;
+			neiReg = 0xf800 + (unsigned char)abDir;
 			neiArea = 0;
 		}
 		else
@@ -241,41 +248,69 @@ static void getNeighbourRegAndArea(dtTileCacheLayer& layer,
 			neiArea = 0;
 		}
 
-		neiHeight = layer.heights[ia];	// UE
+		// Find the vertex height. Try going A-D-C, get height of d and c.
+		vertexHeight = layer.heights[ia];
+		if ((acon & bcDirMask) != 0)
+		{
+			// a is connected to d
+			const int dx = ax + getDirOffsetX(bcDir);
+			const int dy = ay + getDirOffsetY(bcDir);
+			const int id = dx + dy * w;
+			vertexHeight = dtMax(vertexHeight, layer.heights[id]);
+
+			const unsigned char dcon = layer.cons[id] & 0xf;
+			if ((dcon & abDirMask) != 0)
+			{
+				// d is connected to c
+				const int cx = dx + getDirOffsetX(abDir);
+				const int cy = dy + getDirOffsetY(abDir);
+				const int ic = cx + cy * w;
+				vertexHeight = dtMax(vertexHeight, layer.heights[ic]);
+			}
+		}
 	}
 	else
 	{
-		const int bx = ax + getDirOffsetX(dir);
-		const int by = ay + getDirOffsetY(dir);
+		// a is connected to b
+		const int bx = ax + getDirOffsetX(abDir);
+		const int by = ay + getDirOffsetY(abDir);
 		const int ib = bx + by*w;
 
 		neiReg = layer.regs[ib];
 		neiArea = layer.areas[ib];
-		neiHeight = layer.heights[ib];	 // UE
+		vertexHeight = dtMax(layer.heights[ia], layer.heights[ib]);
 
 		// Get area type of the cell diagonal [c] to current cell [a]. Where [b] is direct neighbour in the direction of 'dir'.
 		//   ^
 		//  [b][c]
-		//  [a]
-		const int cdir = (dir + 1) & 0x3;
+		//  [a][d]
 		const unsigned char bcon = layer.cons[ib] & 0xf;
-		const unsigned char bportal = layer.cons[ib] >> 4;
-		const unsigned char bmask = (unsigned char)(1 << cdir);
-		if ((bcon & bmask) == 0)
+		if ((bcon & bcDirMask) == 0)
 		{
 			cornerNeiArea = 0;
 		}
 		else
 		{
-			const int cx = bx + getDirOffsetX(cdir);
-			const int cy = by + getDirOffsetY(cdir);
+			// b is connected to c
+			const int cx = bx + getDirOffsetX(bcDir);
+			const int cy = by + getDirOffsetY(bcDir);
 			const int ic = cx + cy * w;
 			cornerNeiArea = layer.areas[ic];
-			neiHeight = dtMax(neiHeight, layer.heights[ic]);	// UE
+			vertexHeight = dtMax(vertexHeight, layer.heights[ic]);
 		}
 
+		// If connected, check the height of d to get the maximum for the vertexHeight.
+		if ((acon & bcDirMask) != 0)
+		{
+			// a is connected to d
+			const int dx = ax + getDirOffsetX(bcDir);
+			const int dy = ay + getDirOffsetY(bcDir);
+			const int id = dx + dy * w;
+			vertexHeight = dtMax(vertexHeight, layer.heights[id]);
+		}
 	}
 }
+//@UE END
 
 static bool walkContour(dtTileCacheLayer& layer, int x, int y, int idx, const bool allowMerging, unsigned char* flags, dtTempContour& cont, int& contourIndex) // UE
 {
@@ -293,7 +328,7 @@ static bool walkContour(dtTileCacheLayer& layer, int x, int y, int idx, const bo
 	unsigned short neiReg = 0xffff;
 	unsigned char neiArea = 0;
 	unsigned char cornerNeiArea = 0;
-	unsigned short neiHeight = 0;	// UE
+	unsigned short vertexHeight = 0;	// UE
 	unsigned short prevNeiArea = 0;
 	unsigned short prevCornerNeiArea = 0;
 	bool checkForPinning = false;
@@ -306,7 +341,7 @@ static bool walkContour(dtTileCacheLayer& layer, int x, int y, int idx, const bo
 		int ny = y;
 		unsigned char ndir = dir;
 
-		getNeighbourRegAndArea(layer, x, y, dir, neiReg, neiArea, cornerNeiArea, neiHeight);	// UE
+		getNeighbourRegAndAreaAndVertexHeight(layer, x, y, dir, neiReg, neiArea, cornerNeiArea, vertexHeight);	// UE
 
 		if (neiReg != layer.regs[x+y*w])
 		{
@@ -337,8 +372,7 @@ static bool walkContour(dtTileCacheLayer& layer, int x, int y, int idx, const bo
 			}
 
 			// Try to merge with previous vertex.
-			const int py = dtMax(neiHeight, (int)layer.heights[x+y*w]);	// UE
-			if (!appendVertex(cont, px, py, pz, neiReg, neiArea, allowMerging)) // UE
+			if (!appendVertex(cont, px, vertexHeight, pz, neiReg, neiArea, allowMerging)) // UE
 				return false;
 
 			flags[idx] &= ~(1 << dir); // Remove visited edges
