[NavMesh] When removing poly, make sure that the totality of the replacement poly is copied.
[NavMesh] When looping to see if vertex have been flagged for removal, stop the loop to prevent the flag from being reset.

diff --git a/Engine/Source/Runtime/Navmesh/Private/DetourTileCache/DetourTileCacheBuilder.cpp b/Engine/Source/Runtime/Navmesh/Private/DetourTileCache/DetourTileCacheBuilder.cpp
--- a/Engine/Source/Runtime/Navmesh/Private/DetourTileCache/DetourTileCacheBuilder.cpp
+++ b/Engine/Source/Runtime/Navmesh/Private/DetourTileCache/DetourTileCacheBuilder.cpp
@@ -1882,7 +1882,14 @@ static dtStatus removeVertex(dtTileCacheLogContext* ctx, dtTileCachePolyMesh& me
 		const int nv = countPolyVerts(p);
 		bool hasRem = false;
 		for (int j = 0; j < nv; ++j)
-			if (p[j] == rem) hasRem = true;
+		{
+			if (p[j] == rem)
+			{
+				hasRem = true;
+				break;
+			}
+		}
+
 		if (hasRem)
 		{
 			// Collect edges which does not touch the removed vertex.
@@ -1897,10 +1904,10 @@ static dtStatus removeVertex(dtTileCacheLogContext* ctx, dtTileCachePolyMesh& me
 					nedges++;
 				}
 			}
-			// Remove the polygon.
-			const unsigned short* p2 = &mesh.polys[(mesh.npolys-1)*MAX_VERTS_PER_POLY*2];
-			memcpy(p,p2,sizeof(unsigned short)*MAX_VERTS_PER_POLY);
-			memset(p+MAX_VERTS_PER_POLY,0xff,sizeof(unsigned short)*MAX_VERTS_PER_POLY);
+
+			// Remove the polygon p.
+			const unsigned short* lastp = &mesh.polys[(mesh.npolys-1)*MAX_VERTS_PER_POLY*2];
+			memcpy(p, lastp, sizeof(unsigned short)*MAX_VERTS_PER_POLY*2);
 			mesh.areas[i] = mesh.areas[mesh.npolys-1];
 			mesh.npolys--;
 			--i;
@@ -1922,7 +1929,9 @@ static dtStatus removeVertex(dtTileCacheLogContext* ctx, dtTileCachePolyMesh& me
 		unsigned short* p = &mesh.polys[i*MAX_VERTS_PER_POLY*2];
 		const int nv = countPolyVerts(p);
 		for (int j = 0; j < nv; ++j)
+		{
 			if (p[j] > rem) p[j]--;
+		}
 	}
 	for (int i = 0; i < nedges; ++i)
 	{
@@ -1962,6 +1971,7 @@ static dtStatus removeVertex(dtTileCacheLogContext* ctx, dtTileCachePolyMesh& me
 				pushBack(a, harea, nharea);
 				add = true;
 			}
+			
 			if (add)
 			{
 				// The edge segment was added, remove it.
@@ -2091,11 +2101,15 @@ static dtStatus removeVertex(dtTileCacheLogContext* ctx, dtTileCachePolyMesh& me
 	// Store polygons.
 	for (int i = 0; i < npolys; ++i)
 	{
-		if (mesh.npolys >= maxTris) break;
+		if (mesh.npolys >= maxTris)
+			break;
+
 		unsigned short* newPoly = &mesh.polys[mesh.npolys*MAX_VERTS_PER_POLY*2];
 		memset(newPoly,0xff,sizeof(unsigned short)*MAX_VERTS_PER_POLY*2);
 		for (int j = 0; j < MAX_VERTS_PER_POLY; ++j)
+		{
 			newPoly[j] = polys[i*MAX_VERTS_PER_POLY+j];
+		}
 		mesh.areas[mesh.npolys] = pareas[i];
 		mesh.npolys++;
 		if (mesh.npolys > maxTris)
@@ -2124,7 +2138,9 @@ dtStatus dtBuildTileCachePolyMesh(dtTileCacheAlloc* alloc,
 	for (int i = 0; i < lcset.nconts; ++i)
 	{
 		// Skip null contours.
-		if (lcset.conts[i].nverts < 3 || lcset.conts[i].area == DT_TILECACHE_NULL_AREA) continue;
+		if (lcset.conts[i].nverts < 3 || lcset.conts[i].area == DT_TILECACHE_NULL_AREA)
+			continue;
+
 		maxVertices += lcset.conts[i].nverts;
 		maxTris += lcset.conts[i].nverts - 2;
 		maxVertsPerCont = dtMax(maxVertsPerCont, lcset.conts[i].nverts);
@@ -2183,7 +2199,9 @@ dtStatus dtBuildTileCachePolyMesh(dtTileCacheAlloc* alloc,
 	
 	unsigned short firstVert[VERTEX_BUCKET_COUNT2];
 	for (int i = 0; i < VERTEX_BUCKET_COUNT2; ++i)
+	{
 		firstVert[i] = DT_TILECACHE_NULL_IDX;
+	}
 	
 	dtFixedArray<unsigned short> nextVert(alloc, maxVertices);
 	if (!nextVert)
@@ -2326,10 +2344,13 @@ dtStatus dtBuildTileCachePolyMesh(dtTileCacheAlloc* alloc,
 			dtStatus status = removeVertex(ctx, mesh, (unsigned short)i, maxTris);
 			if (dtStatusFailed(status))
 				return status;
+			
 			// Remove vertex
 			// Note: mesh.nverts is already decremented inside removeVertex()!
 			for (int j = i; j < mesh.nverts; ++j)
+			{
 				vflags[j] = vflags[j+1];
+			}
 			--i;
 		}
 	}
