[NavMesh] dtBuildTileCachePolyMesh fix: while removing vertices, fix addition of an extra point to the countour since it can break the triangulation and result in missing polygons.

diff --git a/Engine/Source/Runtime/Navmesh/Private/DetourTileCache/DetourTileCacheBuilder.cpp b/Engine/Source/Runtime/Navmesh/Private/DetourTileCache/DetourTileCacheBuilder.cpp
--- a/Engine/Source/Runtime/Navmesh/Private/DetourTileCache/DetourTileCacheBuilder.cpp
+++ b/Engine/Source/Runtime/Navmesh/Private/DetourTileCache/DetourTileCacheBuilder.cpp
@@ -1955,20 +1955,20 @@ static dtStatus removeVertex(dtTileCacheLogContext* ctx, dtTileCachePolyMesh& me
 		{
 			const unsigned short ea = edges[i*3+0];
 			const unsigned short eb = edges[i*3+1];
-			const unsigned short a = edges[i*3+2];
+			const unsigned short area = edges[i*3+2];
 			bool add = false;
 			if (hole[0] == eb)
 			{
 				// The segment matches the beginning of the hole boundary.
 				pushFront(ea, hole, nhole);
-				pushFront(a, harea, nharea);
+				pushFront(area, harea, nharea);
 				add = true;
 			}
 			else if (hole[nhole-1] == ea)
 			{
 				// The segment matches the end of the hole boundary.
 				pushBack(eb, hole, nhole);
-				pushBack(a, harea, nharea);
+				pushBack(area, harea, nharea);
 				add = true;
 			}
 			
@@ -1987,8 +1987,15 @@ static dtStatus removeVertex(dtTileCacheLogContext* ctx, dtTileCachePolyMesh& me
 		if (!match)
 			break;
 	}
+
+	// Remove duplicate vertex at end.
+	if (nhole > 0)
+	{
+		if (hole[0] == hole[nhole-1])
+			nhole--;
+	}
 	
-	// Skip degenerated areas
+	// Skip degenerated areas.
 	if (nhole < 3)
 		return DT_SUCCESS;
 	
