[NavMesh] Fix tile portals not connecting caused by picking opposite portal direction (on an edge above or below) when building mesh adjacency.
This fixes issues seen in stairs where the stair corridor can be blocked by a hard edges when edges between layers do not match perfectly.

diff --git a/Engine/Source/Runtime/NavigationSystem/Private/NavMesh/RecastNavMeshGenerator.cpp b/Engine/Source/Runtime/NavigationSystem/Private/NavMesh/RecastNavMeshGenerator.cpp
--- a/Engine/Source/Runtime/NavigationSystem/Private/NavMesh/RecastNavMeshGenerator.cpp
+++ b/Engine/Source/Runtime/NavigationSystem/Private/NavMesh/RecastNavMeshGenerator.cpp
@@ -3767,7 +3767,7 @@ bool FRecastTileGenerator::GenerateNavigationDataLayer(FNavMeshBuildContext& Bui
 			return false;
 		}
 
-		status = dtBuildTileCachePolyMesh(&GenNavAllocator, &BuildContext, *GenerationContext.ContourSet, *GenerationContext.PolyMesh);
+		status = dtBuildTileCachePolyMesh(&GenNavAllocator, &BuildContext, *GenerationContext.ContourSet, *GenerationContext.PolyMesh, TileConfig.walkableClimb);
 		if (dtStatusFailed(status))
 		{
 			BuildContext.log(RC_LOG_ERROR, "GenerateNavigationData: Failed to generate poly mesh.");
diff --git a/Engine/Source/Runtime/Navmesh/Private/DetourTileCache/DetourTileCache.cpp b/Engine/Source/Runtime/Navmesh/Private/DetourTileCache/DetourTileCache.cpp
--- a/Engine/Source/Runtime/Navmesh/Private/DetourTileCache/DetourTileCache.cpp
+++ b/Engine/Source/Runtime/Navmesh/Private/DetourTileCache/DetourTileCache.cpp
@@ -718,7 +718,7 @@ dtStatus dtTileCache::buildNavMeshTile(const dtCompressedTileRef ref, dtNavMesh*
 	bc.lmesh = dtAllocTileCachePolyMesh(m_talloc);
 	if (!bc.lmesh)
 		return status;
-	status = dtBuildTileCachePolyMesh(m_talloc, 0, *bc.lcset, *bc.lmesh);
+	status = dtBuildTileCachePolyMesh(m_talloc, 0, *bc.lcset, *bc.lmesh, m_params.walkableClimb);
 	if (dtStatusFailed(status))
 		return status;
 	
diff --git a/Engine/Source/Runtime/Navmesh/Private/DetourTileCache/DetourTileCacheBuilder.cpp b/Engine/Source/Runtime/Navmesh/Private/DetourTileCache/DetourTileCacheBuilder.cpp
--- a/Engine/Source/Runtime/Navmesh/Private/DetourTileCache/DetourTileCacheBuilder.cpp
+++ b/Engine/Source/Runtime/Navmesh/Private/DetourTileCache/DetourTileCacheBuilder.cpp
@@ -164,9 +164,14 @@ struct dtTempContour
 
 
 inline bool overlapRangeExl(const unsigned short amin, const unsigned short amax,
-							const unsigned short bmin, const unsigned short bmax)
+							const unsigned short bmin, const unsigned short bmax,
+							const unsigned short ya, const unsigned short eya,
+							const unsigned short yb, const unsigned short eyb,
+							const int walkableClimb)
 {
-	return (amin >= bmax || amax <= bmin) ? false : true;
+	const bool longitudinalOverlapExl  = !(amin >= bmax || amax <= bmin);
+	const bool elevationOverlap = (dtAbs(ya-eya) <= walkableClimb) || (dtAbs(yb-eyb) <= walkableClimb);
+	return longitudinalOverlapExl && elevationOverlap;
 }
 
 // Returns true on success, false if there was an error adding a vertex.
@@ -1208,8 +1213,8 @@ namespace TileCacheData
 {
 	struct rcEdge
 	{
-		unsigned short vert[2];
-		unsigned short polyEdge[2];
+		unsigned short vert[2];			// index in verts (a,b)
+		unsigned short polyEdge[2];		// index in polys (a,b)
 		unsigned short poly[2];
 	};
 }
@@ -1217,7 +1222,8 @@ namespace TileCacheData
 static bool buildMeshAdjacency(dtTileCacheAlloc* alloc,
 							   unsigned short* polys, const int npolys,
 							   const unsigned short* verts, const int nverts,
-							   const dtTileCacheContourSet& lcset)
+							   const dtTileCacheContourSet& lcset,
+							   const int walkableClimb)
 {
 	// Based on code by Eric Lengyel from:
 	// http://www.terathon.com/code/edges.php
@@ -1263,7 +1269,8 @@ static bool buildMeshAdjacency(dtTileCacheAlloc* alloc,
 			}
 		}
 	}
-	
+
+	// Find matching edges
 	for (int i = 0; i < npolys; ++i)
 	{
 		const unsigned short* t = &polys[i*MAX_VERTS_PER_POLY*2];
@@ -1282,12 +1289,14 @@ static bool buildMeshAdjacency(dtTileCacheAlloc* alloc,
 					TileCacheData::rcEdge& edge = edges[e];
 					if (edge.vert[1] == v0 && edge.poly[0] == edge.poly[1])
 					{
+						// Edges matches
 						edge.poly[1] = (unsigned short)i;
 						edge.polyEdge[1] = (unsigned short)j;
 						found = true;
 						break;
 					}
 				}
+				
 				if (!found)
 				{
 					// Matching edge not found, it is an open edge, add it.
@@ -1298,6 +1307,7 @@ static bool buildMeshAdjacency(dtTileCacheAlloc* alloc,
 					edge.polyEdge[0] = (unsigned short)j;
 					edge.poly[1] = (unsigned short)i;
 					edge.polyEdge[1] = 0xff;
+
 					// Insert edge
 					nextEdge[edgeCount] = firstEdge[v1];
 					firstEdge[v1] = (unsigned short)edgeCount;
@@ -1318,18 +1328,24 @@ static bool buildMeshAdjacency(dtTileCacheAlloc* alloc,
 		{
 			const unsigned short* va = &cont.verts[k*4];
 			const unsigned short* vb = &cont.verts[j*4];
+			
 			const unsigned char dir = va[3] & 0xf;
 			if (dir == 0xf)
 				continue;
 			
 			if (dir == 0 || dir == 2)
 			{
-				// Find matching vertical edge
+				// Find matching edge on z axis
 				const unsigned short x = va[0];
 				unsigned short zmin = va[2];
 				unsigned short zmax = vb[2];
+				unsigned short ya = va[1];
+				unsigned short yb = vb[1];
 				if (zmin > zmax)
+				{
 					dtSwap(zmin, zmax);
+					dtSwap(ya, yb);
+				}
 				
 				for (int m = 0; m < edgeCount; ++m)
 				{
@@ -1337,15 +1353,22 @@ static bool buildMeshAdjacency(dtTileCacheAlloc* alloc,
 					// Skip connected edges.
 					if (e.poly[0] != e.poly[1])
 						continue;
+
 					const unsigned short* eva = &verts[e.vert[0]*3];
 					const unsigned short* evb = &verts[e.vert[1]*3];
 					if (eva[0] == x && evb[0] == x)
 					{
 						unsigned short ezmin = eva[2];
 						unsigned short ezmax = evb[2];
+						unsigned short eya = eva[1];
+						unsigned short eyb = evb[1];
 						if (ezmin > ezmax)
+						{
 							dtSwap(ezmin, ezmax);
-						if (overlapRangeExl(zmin,zmax, ezmin, ezmax))
+							dtSwap(eya, eyb);
+						}
+
+						if (overlapRangeExl(zmin,zmax, ezmin, ezmax, ya, eya, yb, eyb, walkableClimb))
 						{
 							// Reuse the other polyedge to store dir.
 							e.polyEdge[1] = dir;
@@ -1355,27 +1378,40 @@ static bool buildMeshAdjacency(dtTileCacheAlloc* alloc,
 			}
 			else
 			{
-				// Find matching vertical edge
+				// Find matching edge on x axis
 				const unsigned short z = va[2];
 				unsigned short xmin = va[0];
 				unsigned short xmax = vb[0];
+				unsigned short ya = va[1];
+				unsigned short yb = vb[1];
 				if (xmin > xmax)
+				{
 					dtSwap(xmin, xmax);
+					dtSwap(ya, yb);
+				}
+				
 				for (int m = 0; m < edgeCount; ++m)
 				{
 					TileCacheData::rcEdge& e = edges[m];
 					// Skip connected edges.
 					if (e.poly[0] != e.poly[1])
 						continue;
+
 					const unsigned short* eva = &verts[e.vert[0]*3];
 					const unsigned short* evb = &verts[e.vert[1]*3];
 					if (eva[2] == z && evb[2] == z)
 					{
 						unsigned short exmin = eva[0];
 						unsigned short exmax = evb[0];
+						unsigned short eya = eva[1];
+						unsigned short eyb = evb[1];
 						if (exmin > exmax)
+						{
 							dtSwap(exmin, exmax);
-						if (overlapRangeExl(xmin,xmax, exmin, exmax))
+							dtSwap(eya, eyb);
+						}
+
+						if (overlapRangeExl(xmin,xmax, exmin, exmax, ya, eya, yb, eyb, walkableClimb))
 						{
 							// Reuse the other polyedge to store dir.
 							e.polyEdge[1] = dir;
@@ -1606,8 +1642,11 @@ static int triangulate(int n, const unsigned short* verts, unsigned short* indic
 		for (int k = i1; k < n; k++)
 			indices[k] = indices[k+1];
 		
-		if (i1 >= n) i1 = 0;
+		if (i1 >= n)
+			i1 = 0;
+
 		i = TileCacheFunc::prev(i1, n);
+		
 		// Update diagonal flags.
 		if (diagonal(TileCacheFunc::prev(i, n), i1, n, verts, indices))
 			indices[i] |= 0x8000;
@@ -1647,7 +1686,7 @@ namespace TileCacheFunc
 	}
 }
 
-static int getPolyMergeValue(unsigned short* pa, unsigned short* pb,
+static int getPolyMergeValue(const unsigned short* pa, const unsigned short* pb,
 							 const unsigned short* verts, int& ea, int& eb)
 {
 	const int na = countPolyVerts(pa);
@@ -1930,9 +1969,11 @@ static dtStatus removeVertex(dtTileCacheLogContext* ctx, dtTileCachePolyMesh& me
 		const int nv = countPolyVerts(p);
 		for (int j = 0; j < nv; ++j)
 		{
-			if (p[j] > rem) p[j]--;
+			if (p[j] > rem)
+				p[j]--;
 		}
 	}
+	
 	for (int i = 0; i < nedges; ++i)
 	{
 		if (edges[i*3+0] > rem) edges[i*3+0]--;
@@ -2132,9 +2173,10 @@ static dtStatus removeVertex(dtTileCacheLogContext* ctx, dtTileCachePolyMesh& me
 
 
 dtStatus dtBuildTileCachePolyMesh(dtTileCacheAlloc* alloc, 
-								  dtTileCacheLogContext* ctx,
-								  dtTileCacheContourSet& lcset,
-								  dtTileCachePolyMesh& mesh)
+                                  dtTileCacheLogContext* ctx,
+                                  dtTileCacheContourSet& lcset,
+                                  dtTileCachePolyMesh& mesh,
+                                  const int walkableClimb)
 {
 	dtAssert(alloc);
 	
@@ -2272,11 +2314,12 @@ dtStatus dtBuildTileCachePolyMesh(dtTileCacheAlloc* alloc,
 				npolys++;
 			}
 		}
+		
 		if (!npolys)
 			continue;
 		
 		// Merge polygons.
-		int maxVertsPerPoly =MAX_VERTS_PER_POLY ;
+		int maxVertsPerPoly = MAX_VERTS_PER_POLY;
 		if (maxVertsPerPoly > 3) //-V547
 		{
 			for(;;)
@@ -2363,7 +2406,7 @@ dtStatus dtBuildTileCachePolyMesh(dtTileCacheAlloc* alloc,
 	}
 	
 	// Calculate adjacency.
-	if (!buildMeshAdjacency(alloc, mesh.polys, mesh.npolys, mesh.verts, mesh.nverts, lcset))
+	if (!buildMeshAdjacency(alloc, mesh.polys, mesh.npolys, mesh.verts, mesh.nverts, lcset, walkableClimb))
 		return DT_FAILURE | DT_OUT_OF_MEMORY;
 		
 	return DT_SUCCESS;
diff --git a/Engine/Source/Runtime/Navmesh/Public/DetourTileCache/DetourTileCacheBuilder.h b/Engine/Source/Runtime/Navmesh/Public/DetourTileCache/DetourTileCacheBuilder.h
--- a/Engine/Source/Runtime/Navmesh/Public/DetourTileCache/DetourTileCacheBuilder.h
+++ b/Engine/Source/Runtime/Navmesh/Public/DetourTileCache/DetourTileCacheBuilder.h
@@ -58,7 +58,7 @@ struct dtTileCacheLayer
 
 struct dtTileCacheContour
 {
-	unsigned short* verts;
+	unsigned short* verts;	///< x, y, z, portal direction.
 	int nverts;
 	unsigned short reg;
 	unsigned char area;
@@ -277,7 +277,8 @@ NAVMESH_API dtStatus dtBuildTileCacheContours(dtTileCacheAlloc* alloc,
 NAVMESH_API dtStatus dtBuildTileCachePolyMesh(dtTileCacheAlloc* alloc,
 								dtTileCacheLogContext* ctx,
 								dtTileCacheContourSet& lcset,
-								dtTileCachePolyMesh& mesh);
+								dtTileCachePolyMesh& mesh,
+								const int walkableClimb);
 
 //@UE BEGIN
 NAVMESH_API dtStatus dtBuildTileCachePolyMeshDetail(dtTileCacheAlloc* alloc,
