[NavMesh] Optimization to simplifyContour

diff --git a/Engine/Source/Runtime/Navmesh/Private/DetourTileCache/DetourTileCacheBuilder.cpp b/Engine/Source/Runtime/Navmesh/Private/DetourTileCache/DetourTileCacheBuilder.cpp
--- a/Engine/Source/Runtime/Navmesh/Private/DetourTileCache/DetourTileCacheBuilder.cpp
+++ b/Engine/Source/Runtime/Navmesh/Private/DetourTileCache/DetourTileCacheBuilder.cpp
@@ -531,6 +531,7 @@ static void simplifyContour(unsigned char area, unsigned short region, dtTempCon
 	}
 
 	const dtReal heightRatio = elevationRatio * ch / cs; // UE
+	const bool checkElevation = elevationRatio > 0;		 // UE
 
 	// Add points until all raw points are within
 	// error tolerance to the simplified shape.
@@ -553,6 +554,20 @@ static void simplifyContour(unsigned char area, unsigned short region, dtTempCon
 		int maxi = -1;
 		int ci, cinc, endi;
 
+//@UE BEGIN
+		dtReal a[3];
+		dtReal b[3];
+		if (checkElevation)
+		{
+			a[0] = (dtReal)ax;
+			a[1] = heightRatio * ay;
+			a[2] = (dtReal)az;
+			b[0] = (dtReal)bx;
+			b[1] = heightRatio * by;
+			b[2] = (dtReal)bz;
+		}
+//@UE END
+		
 		// Traverse the segment in lexilogical order so that the
 		// max deviation is calculated similarly when traversing
 		// opposite segments.
@@ -568,31 +583,38 @@ static void simplifyContour(unsigned char area, unsigned short region, dtTempCon
 			ci = (bi+cinc) % cont.nverts;
 			endi = ai;
 
+//@UE BEGIN
 			// Because of floating point imprecision, dtDistancePtSegSqr to a-b might be slightly differ from dtDistancePtSegSqr to b-a.
 			// Swap points because we need the maximum deviation to be computed the same way for opposite segments to match.
-			dtSwap(ax, bx);	// UE
-			dtSwap(ay, by);	// UE
-			dtSwap(az, bz);	// UE
+			if (checkElevation)
+			{
+				dtSwap(a[0], b[0]);
+				dtSwap(a[1], b[1]);
+				dtSwap(a[2], b[2]);
+			}
+			else
+			{
+				dtSwap(ax, bx);
+				dtSwap(az, bz);
+			}
+//@UE END
 		}
 
 		// Tessellate only between regions and areas.
 		const unsigned short* ciSrc = &cont.verts[ci*5];
 		const int ciReg = ciSrc[3];
 		const unsigned char ciArea = (unsigned char)ciSrc[4];
-		const bool checkRegionChange = elevationRatio > 0;								 // UE
-		if (area != ciArea || ciReg == 0xffff || (checkRegionChange && region != ciReg)) // UE
+		if (area != ciArea || ciReg == 0xffff || (checkElevation && region != ciReg)) // UE
 		{
 			while (ci != endi)
 			{
 //@UE BEGIN
 				dtReal d;
-				if (elevationRatio > 0)
+				if (checkElevation)
 				{
 					// Instead of multiplying all components by ch or cs to go from voxels to world units, 
 					// we just use the heightRatio (avoiding extra cs multiplication on x and z).
 					const dtReal pt[3] = { (dtReal)cont.verts[ci*5+0], heightRatio*cont.verts[ci*5+1], (dtReal)cont.verts[ci*5+2] };
-					const dtReal a[3] = { (dtReal)ax, heightRatio*ay, (dtReal)az };
-					const dtReal b[3] = { (dtReal)bx, heightRatio*by, (dtReal)bz };
 					d = dtDistancePtSegSqr(pt, a, b);
 				}
 				else
