[NavMesh] Navlink generation first pass (experimental).
When enabled, generates simple links within navmesh tiles.
Added debug display showing trajectory and trajectory validation.

diff --git a/Engine/Source/Runtime/Engine/Public/EngineStats.h b/Engine/Source/Runtime/Engine/Public/EngineStats.h
--- a/Engine/Source/Runtime/Engine/Public/EngineStats.h
+++ b/Engine/Source/Runtime/Engine/Public/EngineStats.h
@@ -140,6 +140,7 @@ DECLARE_CYCLE_STAT_EXTERN(TEXT("Recast: build layers"),STAT_Navigation_RecastBui
 DECLARE_CYCLE_STAT_EXTERN(TEXT("Recast: build tile cache"),STAT_Navigation_RecastBuildTileCache,STATGROUP_Navigation, );
 DECLARE_CYCLE_STAT_EXTERN(TEXT("Recast: build navmesh"),STAT_Navigation_RecastBuildNavigation,STATGROUP_Navigation, );
 DECLARE_CYCLE_STAT_EXTERN(TEXT("Recast: build navigation layer"), STAT_Navigation_GenerateNavigationDataLayer, STATGROUP_Navigation, );
+DECLARE_CYCLE_STAT_EXTERN(TEXT("Recast: build links"), STAT_Navigation_RecastBuildLinks, STATGROUP_Navigation, );
 DECLARE_CYCLE_STAT_EXTERN(TEXT("Recast: build regions"), STAT_Navigation_RecastBuildRegions, STATGROUP_Navigation, );
 DECLARE_CYCLE_STAT_EXTERN(TEXT("Recast: mark areas"), STAT_Navigation_RecastMarkAreas, STATGROUP_Navigation, );
 DECLARE_CYCLE_STAT_EXTERN(TEXT("Recast: build contours"), STAT_Navigation_RecastBuildContours, STATGROUP_Navigation, );
@@ -171,7 +172,7 @@ DECLARE_MEMORY_STAT_EXTERN(TEXT("Detour PERM_TILE_DATA"), STAT_Navigation_Detour
 DECLARE_MEMORY_STAT_EXTERN(TEXT("Detour PERM_TILE_DYNLINK_OFFMESH"), STAT_Navigation_DetourPERM_TILE_DYNLINK_OFFMESH, STATGROUP_Navigation, );
 DECLARE_MEMORY_STAT_EXTERN(TEXT("Detour PERM_TILE_DYNLINK_CLUSTER"), STAT_Navigation_DetourPERM_TILE_DYNLINK_CLUSTER, STATGROUP_Navigation, );
 DECLARE_MEMORY_STAT_EXTERN(TEXT("Detour PERM_TILES"), STAT_Navigation_DetourPERM_TILES, STATGROUP_Navigation, );
-
+DECLARE_MEMORY_STAT_EXTERN(TEXT("Detour PERM_TILE_LINK_BUILDER"), STAT_Navigation_DetourPERM_TILE_LINK_BUILDER, STATGROUP_Navigation, );
 
 DECLARE_MEMORY_STAT_EXTERN(TEXT("Detour Tile Memory"), STAT_DetourTileMemory, STATGROUP_Navigation, );
 DECLARE_MEMORY_STAT_EXTERN(TEXT("Detour Tile Mesh Header Memory"), STAT_DetourTileMeshHeaderMemory, STATGROUP_Navigation, );
diff --git a/Engine/Source/Runtime/NavigationSystem/Private/NavMesh/PImplRecastNavMesh.cpp b/Engine/Source/Runtime/NavigationSystem/Private/NavMesh/PImplRecastNavMesh.cpp
--- a/Engine/Source/Runtime/NavigationSystem/Private/NavMesh/PImplRecastNavMesh.cpp
+++ b/Engine/Source/Runtime/NavigationSystem/Private/NavMesh/PImplRecastNavMesh.cpp
@@ -109,6 +109,10 @@ static void* DetourMalloc(int Size, dtAllocHint Hint)
 		INC_MEMORY_STAT_BY(STAT_Navigation_DetourPERM_TILES, ActualSize);
 		break;
 
+	case DT_ALLOC_PERM_TILE_LINK_BUILDER:
+		INC_MEMORY_STAT_BY(STAT_Navigation_DetourPERM_TILE_LINK_BUILDER, ActualSize);
+		break;
+		
 	default:
 		ensureMsgf(false, TEXT("Unsupported allocation hint %d"), Hint);
 		break;
@@ -195,6 +199,10 @@ static void DetourFree(void* Original, dtAllocHint Hint)
 		DEC_MEMORY_STAT_BY(STAT_Navigation_DetourPERM_TILES, Size);
 		break;
 
+	case DT_ALLOC_PERM_TILE_LINK_BUILDER:
+		DEC_MEMORY_STAT_BY(STAT_Navigation_DetourPERM_TILE_LINK_BUILDER, Size);
+		break;
+		
 	default:
 		ensureMsgf(false, TEXT("Unsupported allocation hint %d"), Hint);
 		break;
diff --git a/Engine/Source/Runtime/NavigationSystem/Private/NavMesh/RecastNavMesh.cpp b/Engine/Source/Runtime/NavigationSystem/Private/NavMesh/RecastNavMesh.cpp
--- a/Engine/Source/Runtime/NavigationSystem/Private/NavMesh/RecastNavMesh.cpp
+++ b/Engine/Source/Runtime/NavigationSystem/Private/NavMesh/RecastNavMesh.cpp
@@ -506,6 +506,8 @@ FRecastNavMeshTileGenerationDebug::FRecastNavMeshTileGenerationDebug()
 	bTileCacheContours = false;
 	bTileCachePolyMesh = false;
 	bTileCacheDetailMesh = false;
+	LinkGenerationDebugFlags = 0;
+	LinkGenerationSelectedEdge = -1;
 }
 
 #if WITH_RECAST
@@ -529,6 +531,7 @@ ARecastNavMesh::ARecastNavMesh(const FObjectInitializer& ObjectInitializer)
 	, DefaultMaxHierarchicalSearchNodes(RECAST_MAX_SEARCH_NODES)
 	, bSortNavigationAreasByCost(true)
 	, bIsWorldPartitioned(false)
+	, bGenerateNavLinks(false)
 	, bPerformVoxelFiltering(true)	
 	, bMarkLowHeightAreas(false)
 	, bUseExtraTopCellWhenMarkingAreas(true)
diff --git a/Engine/Source/Runtime/NavigationSystem/Private/NavMesh/RecastNavMeshGenerator.cpp b/Engine/Source/Runtime/NavigationSystem/Private/NavMesh/RecastNavMeshGenerator.cpp
--- a/Engine/Source/Runtime/NavigationSystem/Private/NavMesh/RecastNavMeshGenerator.cpp
+++ b/Engine/Source/Runtime/NavigationSystem/Private/NavMesh/RecastNavMeshGenerator.cpp
@@ -26,6 +26,7 @@
 
 // recast includes
 #include "Detour/DetourNavMeshBuilder.h"
+#include "Detour/DetourNavLinkBuilder.h"
 #include "DetourTileCache/DetourTileCacheBuilder.h"
 #include "NavMesh/RecastHelpers.h"
 #include "NavAreas/NavArea_LowHeight.h"
@@ -38,6 +39,7 @@
 #include "DebugUtils/DebugDraw.h"
 #include "DebugUtils/RecastDebugDraw.h"
 #include "DebugUtils/DetourDebugDraw.h"
+#include "DebugUtils/DetourNavLinkDebugDraw.h"
 #endif //RECAST_INTERNAL_DEBUG_DATA
 
 #ifndef OUTPUT_NAV_TILE_LAYER_COMPRESSION_DATA
@@ -2476,6 +2478,10 @@ ETimeSliceWorkResult FRecastTileGenerator::GenerateTileTimeSliced()
 	FNavMeshBuildContext BuildContext(*this);
 	ETimeSliceWorkResult WorkResult = ETimeSliceWorkResult::Succeeded;
 
+	UE_CLOG(TileConfig.bGenerateLinks, LogNavigation, Warning, TEXT("Generating links in time slice mode is currently not supported so it's been disabled. Turn off bGenerateNavLinks to avoid the warning."));
+	dtLinkBuilderData LinkBuiderData;
+	LinkBuiderData.generatingLinks = false; // Make sure it's disable for now.
+	
 	check(TimeSliceManager);
 
 	switch (GenerateTileTimeSlicedState)
@@ -2516,7 +2522,7 @@ ETimeSliceWorkResult FRecastTileGenerator::GenerateTileTimeSliced()
 	} //fall through to next state
 	case EGenerateTileTimeSlicedState::GenerateNavigationData:
 	{
-		WorkResult = GenerateNavigationDataTimeSliced(BuildContext);
+		WorkResult = GenerateNavigationDataTimeSliced(BuildContext, LinkBuiderData);
 
 		if (WorkResult != ETimeSliceWorkResult::CallAgainNextTimeSlice)
 		{
@@ -2535,7 +2541,6 @@ ETimeSliceWorkResult FRecastTileGenerator::GenerateTileTimeSliced()
 	return WorkResult;
 }
 
-
 bool FRecastTileGenerator::GenerateTile()
 {
 #if RECAST_INTERNAL_DEBUG_DATA
@@ -2548,11 +2553,16 @@ bool FRecastTileGenerator::GenerateTile()
 	FNavMeshBuildContext BuildContext(*this);
 	bool bSuccess = true;
 
+	dtLinkBuilderData LinkBuiderData;
+	LinkBuiderData.generatingLinks = TileConfig.bGenerateLinks;
+
 	if (bRegenerateCompressedLayers)
 	{
 		CompressedLayers.Reset();
 
-		bSuccess = GenerateCompressedLayers(BuildContext);
+		// CompactHF is part of FTileRasterizationContext RasterContext that gets deleted after the GenerateCompressedLayers
+		// so if we are generating links, keep the chf.
+		bSuccess = GenerateCompressedLayers(BuildContext, LinkBuiderData);
 
 #if RECAST_INTERNAL_DEBUG_DATA
 		PostCompressLayerStamp = FPlatformTime::Seconds();
@@ -2567,9 +2577,12 @@ bool FRecastTileGenerator::GenerateTile()
 
 	if (bSuccess)
 	{
-		bSuccess = GenerateNavigationData(BuildContext);
+		bSuccess = GenerateNavigationData(BuildContext, LinkBuiderData);
 	}
 
+	rcFreeHeightField(LinkBuiderData.solidHF);
+	rcFreeCompactHeightfield(LinkBuiderData.compactHF);
+	
 #if RECAST_INTERNAL_DEBUG_DATA	
 	const double EndStamp = FPlatformTime::Seconds();
 	BuildContext.InternalDebugData.BuildTime = EndStamp - StartStamp;
@@ -3460,6 +3473,34 @@ ETimeSliceWorkResult FRecastTileGenerator::GenerateCompressedLayersTimeSliced(FN
 	return ETimeSliceWorkResult::Succeeded;
 }
 
+// Deprecated
+bool FRecastTileGenerator::GenerateCompressedLayers(FNavMeshBuildContext& BuildContext)
+{
+	dtLinkBuilderData LinBuilderData;
+	return GenerateCompressedLayers(BuildContext, LinBuilderData);
+}
+
+// Deprecated
+bool FRecastTileGenerator::GenerateNavigationDataLayer(FNavMeshBuildContext& BuildContext, FTileCacheCompressor& TileCompressor, FTileCacheAllocator& GenNavAllocator, FTileGenerationContext& GenerationContext, int32 LayerIdx)
+{
+	dtLinkBuilderData LinBuilderData;
+	return GenerateNavigationDataLayer(BuildContext, TileCompressor, GenNavAllocator, GenerationContext, LinBuilderData, LayerIdx);
+}
+
+// Deprecated
+ETimeSliceWorkResult FRecastTileGenerator::GenerateNavigationDataTimeSliced(FNavMeshBuildContext& BuildContext)
+{
+	dtLinkBuilderData LinBuilderData;
+	return GenerateNavigationDataTimeSliced(BuildContext, LinBuilderData);
+}
+
+// Deprecated
+bool FRecastTileGenerator::GenerateNavigationData(FNavMeshBuildContext& BuildContext)
+{
+	dtLinkBuilderData LinBuilderData;
+	return GenerateNavigationData(BuildContext, LinBuilderData);
+}
+
 namespace UE::NavMesh::Private
 {
 #if RECAST_INTERNAL_DEBUG_DATA
@@ -3477,11 +3518,10 @@ namespace UE::NavMesh::Private
 #endif //RECAST_INTERNAL_DEBUG_DATA
 };
 
-bool FRecastTileGenerator::GenerateCompressedLayers(FNavMeshBuildContext& BuildContext)
+bool FRecastTileGenerator::GenerateCompressedLayers(FNavMeshBuildContext& BuildContext, dtLinkBuilderData& OutLinkBuilderData)
 {
 	SCOPE_CYCLE_COUNTER(STAT_Navigation_RecastBuildCompressedLayers);
 
-
 	FTileRasterizationContext RasterContext;
 	CompressedLayers.Reset();
 
@@ -3578,6 +3618,17 @@ bool FRecastTileGenerator::GenerateCompressedLayers(FNavMeshBuildContext& BuildC
 		}
 	}
 #endif
+
+	if (OutLinkBuilderData.generatingLinks)
+	{
+		OutLinkBuilderData.solidHF = RasterContext.SolidHF;
+		OutLinkBuilderData.compactHF = RasterContext.CompactHF;
+
+		// When generating links, set SolidHF and CompactHF to null to prevent free 
+		// on the destruction of RasterContext since we will need those later.
+		RasterContext.SolidHF = nullptr;
+		RasterContext.CompactHF = nullptr;
+	}
 	
 	return RecastBuildTileCache(BuildContext, RasterContext);
 }
@@ -3625,14 +3676,89 @@ struct FTileGenerationContext
 	TArray<FNavMeshTileData> NavigationData;
 };
 
-bool FRecastTileGenerator::GenerateNavigationDataLayer(FNavMeshBuildContext& BuildContext, FTileCacheCompressor& TileCompressor, FTileCacheAllocator& GenNavAllocator, FTileGenerationContext& GenerationContext, int32 LayerIdx)
+dtStatus FRecastTileGenerator::BuildTileCacheLinks(FNavMeshBuildContext& BuildContext, dtTileCacheAlloc* alloc, const dtTileCacheLayer& layer,
+	const dtTileCacheContourSet& lcset, const dtLinkBuilderData& linkBuilderData, TArray<FNavigationLink>& OutGeneratedLinks) const
+{
+	duDebugDraw* dd = nullptr;
+	int32 DebugEdge = -1;
+
+#if RECAST_INTERNAL_DEBUG_DATA
+	if (IsTileDebugActive() && (TileDebugSettings.LinkGenerationDebugFlags != 0))
+	{
+		DebugEdge = TileDebugSettings.LinkGenerationSelectedEdge;
+		dd = &BuildContext.InternalDebugData;
+	}
+	const uint16 DebugFlags = TileDebugSettings.LinkGenerationDebugFlags;
+#endif // RECAST_INTERNAL_DEBUG_DATA
+
+	dtAssert(alloc);
+
+	if (!linkBuilderData.solidHF || !linkBuilderData.compactHF)
+	{
+		return DT_FAILURE;
+	}
+
+	const dtReal* orig = layer.header->bmin;
+	dtLinkBuilderConfig linkBuilderConfig;
+	linkBuilderConfig.agentRadius = TileConfig.walkableRadius * TileConfig.cs;
+	linkBuilderConfig.agentHeight = TileConfig.walkableHeight * TileConfig.ch;
+	linkBuilderConfig.agentClimb = TileConfig.walkableClimb * TileConfig.ch;
+	linkBuilderConfig.cellSize = TileConfig.cs;
+	linkBuilderConfig.cellHeight = TileConfig.ch;
+
+	dtNavLinkBuilder linkBuilder;
+	if (!linkBuilder.findEdges(BuildContext, TileConfig, linkBuilderConfig, lcset, orig, linkBuilderData))
+	{
+		return DT_FAILURE;
+	}
+
+	if (DebugEdge == -1)
+	{
+		linkBuilder.buildForAllEdges(linkBuilderConfig, DT_LINK_ACTION_JUMP_DOWN);
+		linkBuilder.buildForAllEdges(linkBuilderConfig, DT_LINK_ACTION_JUMP_OVER);
+
+#if RECAST_INTERNAL_DEBUG_DATA
+		duDebugDrawNavLinkBuilder(dd, linkBuilder, DebugFlags, nullptr);
+	}
+	else
+	{
+		dtNavLinkBuilder::EdgeSampler sampler1;
+		linkBuilder.debugBuildEdge(linkBuilderConfig, DT_LINK_ACTION_JUMP_DOWN, DebugEdge, sampler1);
+		duDebugDrawNavLinkBuilder(dd, linkBuilder, DebugFlags, &sampler1);
+		
+		dtNavLinkBuilder::EdgeSampler sampler2;
+		linkBuilder.debugBuildEdge(linkBuilderConfig, DT_LINK_ACTION_JUMP_OVER, DebugEdge, sampler2);
+		duDebugDrawNavLinkBuilder(dd, linkBuilder, DebugFlags, &sampler2);
+#endif // RECAST_INTERNAL_DEBUG_DATA		
+	}
+
+	// Make FNavigationLinks
+	for (int i = 0; i < linkBuilder.m_nlinks; ++i)
+	{
+		const dtNavLinkBuilder::JumpLink* link = &linkBuilder.m_links[i];
+		
+		if (link->flags == dtNavLinkBuilder::INVALID)
+			continue;
+
+		// For now, just make a link using the center of the range.
+		dtReal midA[3];
+		dtVlerp(midA, &link->spine0[0], &link->spine1[0], 0.5);
+		dtReal midB[3];
+		dtVlerp(midB, &link->spine0[(link->nspine-1)*3], &link->spine1[(link->nspine-1)*3], 0.5);
+		OutGeneratedLinks.Add(FNavigationLink(Recast2UnrealPoint(midA), Recast2UnrealPoint(midB)));
+	}
+	
+	return DT_SUCCESS;
+}
+
+bool FRecastTileGenerator::GenerateNavigationDataLayer(FNavMeshBuildContext& BuildContext, FTileCacheCompressor& TileCompressor,
+	FTileCacheAllocator& GenNavAllocator, FTileGenerationContext& GenerationContext, dtLinkBuilderData& InOutLinkBuilderData, int32 LayerIdx)
 {
 	SCOPE_CYCLE_COUNTER(STAT_Navigation_GenerateNavigationDataLayer)
 		
 	dtStatus status = DT_SUCCESS;
 
 	FNavMeshTileData& CompressedData = CompressedLayers[LayerIdx];
-	const dtTileCacheLayerHeader* TileHeader = (const dtTileCacheLayerHeader*)CompressedData.GetData();
 	GenerationContext.ResetIntermediateData();
 
 	// Decompress tile layer data. 
@@ -3804,6 +3930,7 @@ bool FRecastTileGenerator::GenerateNavigationDataLayer(FNavMeshBuildContext& Bui
 			return false;
 		}
 
+		// Fills GenerationContext.PolyMesh with connectivity information.
 		status = dtBuildTileCachePolyMeshDetail(&GenNavAllocator, TileConfig.cs, TileConfig.ch, TileConfig.detailSampleDist, TileConfig.detailSampleMaxError,
 			*GenerationContext.Layer, *GenerationContext.PolyMesh, *GenerationContext.DetailMesh);
 		if (dtStatusFailed(status))
@@ -3820,6 +3947,21 @@ bool FRecastTileGenerator::GenerateNavigationDataLayer(FNavMeshBuildContext& Bui
 #endif
 	}
 
+	// Build Links
+	TArray<FNavigationLink> GeneratedLinks;
+	if (InOutLinkBuilderData.generatingLinks)
+	{
+		SCOPE_CYCLE_COUNTER(STAT_Navigation_RecastBuildLinks);
+
+		status = BuildTileCacheLinks(BuildContext, &GenNavAllocator, *GenerationContext.Layer, *GenerationContext.ContourSet, InOutLinkBuilderData, GeneratedLinks);
+
+		if (dtStatusFailed(status))
+		{
+			BuildContext.log(RC_LOG_ERROR, "GenerateNavigationDataLayer: Failed to generate links (0x%08X).", status);
+			return false;
+		}
+	}
+
 	unsigned char* NavData = nullptr;
 	int32 NavDataSize = 0;
 
@@ -3834,17 +3976,18 @@ bool FRecastTileGenerator::GenerateNavigationDataLayer(FNavMeshBuildContext& Bui
 			return false;
 		}
 
+		const float DefaultSnapHeight = static_cast<float>(TileConfig.walkableClimb) * static_cast<float>(TileConfig.ch);
+		
 		// if we didn't fail already then it's high time we created data for off-mesh links
 		FOffMeshData OffMeshData;
-		if (OffmeshLinks.Num() > 0)
+		if (!OffmeshLinks.IsEmpty() || !GeneratedLinks.IsEmpty())
 		{
 			SCOPE_CYCLE_COUNTER(STAT_Navigation_RecastGatherOffMeshData);
 
-			OffMeshData.Reserve(OffmeshLinks.Num());
+			OffMeshData.Reserve(OffmeshLinks.Num() + GeneratedLinks.Num());
 			OffMeshData.AreaClassToIdMap = &AdditionalCachedData.AreaClassToIdMap;
 			OffMeshData.FlagsPerArea = AdditionalCachedData.FlagsPerOffMeshLinkArea;
 			const FSimpleLinkNavModifier* LinkModifier = OffmeshLinks.GetData();
-			const float DefaultSnapHeight = static_cast<float>(TileConfig.walkableClimb) * static_cast<float>(TileConfig.ch);
 
 			for (int32 LinkModifierIndex = 0; LinkModifierIndex < OffmeshLinks.Num(); ++LinkModifierIndex, ++LinkModifier)
 			{
@@ -3853,6 +3996,8 @@ bool FRecastTileGenerator::GenerateNavigationDataLayer(FNavMeshBuildContext& Bui
 				OffMeshData.AddSegmentLinks(LinkModifier->SegmentLinks, LinkModifier->LocalToWorld, TileConfig.AgentIndex, DefaultSnapHeight);
 #endif // WITH_NAVMESH_SEGMENT_LINKS
 			}
+
+			OffMeshData.AddLinks(GeneratedLinks, FTransform::Identity, TileConfig.AgentIndex, DefaultSnapHeight);
 		}
 
 		// fill flags, or else detour won't be able to find polygons
@@ -3921,7 +4066,7 @@ bool FRecastTileGenerator::GenerateNavigationDataLayer(FNavMeshBuildContext& Bui
 	return true;
 }
 
-ETimeSliceWorkResult FRecastTileGenerator::GenerateNavigationDataTimeSliced(FNavMeshBuildContext& BuildContext)
+ETimeSliceWorkResult FRecastTileGenerator::GenerateNavigationDataTimeSliced(FNavMeshBuildContext& BuildContext, dtLinkBuilderData& InOutLinkBuilderData)
 {
 	SCOPE_CYCLE_COUNTER(STAT_Navigation_RecastBuildNavigation);
 
@@ -3963,7 +4108,7 @@ ETimeSliceWorkResult FRecastTileGenerator::GenerateNavigationDataTimeSliced(FNav
 				break;
 			}
 
-			const bool bGenDataLayer = GenerateNavigationDataLayer(BuildContext, TileCompressor, *GenNavDataTimeSlicedAllocator, *GenNavDataTimeSlicedGenerationContext, GenNavDataLayerTimeSlicedIdx);
+			const bool bGenDataLayer = GenerateNavigationDataLayer(BuildContext, TileCompressor, *GenNavDataTimeSlicedAllocator, *GenNavDataTimeSlicedGenerationContext, InOutLinkBuilderData, GenNavDataLayerTimeSlicedIdx);
 
 			MARK_TIMESLICE_SECTION_DEBUG(TimeSliceManager->GetTimeSlicer(), GenerateLayers);
 
@@ -4001,7 +4146,7 @@ ETimeSliceWorkResult FRecastTileGenerator::GenerateNavigationDataTimeSliced(FNav
 	return WorkResult;
 }
 
-bool FRecastTileGenerator::GenerateNavigationData(FNavMeshBuildContext& BuildContext)
+bool FRecastTileGenerator::GenerateNavigationData(FNavMeshBuildContext& BuildContext, dtLinkBuilderData& LinkBuilderData)
 {
 	SCOPE_CYCLE_COUNTER(STAT_Navigation_RecastBuildNavigation);
 
@@ -4020,7 +4165,7 @@ bool FRecastTileGenerator::GenerateNavigationData(FNavMeshBuildContext& BuildCon
 			continue;
 		}
 
-		bGenDataLayer = GenerateNavigationDataLayer(BuildContext, TileCompressor, GenNavAllocator, GenerationContext, LayerIdx);
+		bGenDataLayer = GenerateNavigationDataLayer(BuildContext, TileCompressor, GenNavAllocator, GenerationContext, LinkBuilderData, LayerIdx);
 
 		if (!bGenDataLayer)
 		{
@@ -4748,6 +4893,8 @@ void FRecastNavMeshGenerator::ConfigureBuildProperties(FRecastBuildConfig& OutCo
 		OutConfig.walkableHeight = 1;
 	}
 
+	OutConfig.bGenerateLinks = DestNavMesh->bGenerateNavLinks;
+
 	const UNavigationSystemV1* NavSys = FNavigationSystem::GetCurrent<UNavigationSystemV1>(GetWorld());
 	OutConfig.AgentIndex = NavSys ? NavSys->GetSupportedAgentIndex(DestNavMesh) : 0;
 
diff --git a/Engine/Source/Runtime/NavigationSystem/Private/NavigationSystem.cpp b/Engine/Source/Runtime/NavigationSystem/Private/NavigationSystem.cpp
--- a/Engine/Source/Runtime/NavigationSystem/Private/NavigationSystem.cpp
+++ b/Engine/Source/Runtime/NavigationSystem/Private/NavigationSystem.cpp
@@ -117,6 +117,7 @@ DEFINE_STAT(STAT_Navigation_RecastMarkAreas);
 DEFINE_STAT(STAT_Navigation_RecastBuildContours);
 DEFINE_STAT(STAT_Navigation_RecastBuildNavigation);
 DEFINE_STAT(STAT_Navigation_GenerateNavigationDataLayer);
+DEFINE_STAT(STAT_Navigation_RecastBuildLinks);
 DEFINE_STAT(STAT_Navigation_RecastBuildRegions);
 DEFINE_STAT(STAT_Navigation_UpdateNavOctree);
 DEFINE_STAT(STAT_Navigation_CollisionTreeMemory);
@@ -146,6 +147,7 @@ DEFINE_STAT(STAT_Navigation_DetourPERM_TILE_DATA);
 DEFINE_STAT(STAT_Navigation_DetourPERM_TILE_DYNLINK_OFFMESH);
 DEFINE_STAT(STAT_Navigation_DetourPERM_TILE_DYNLINK_CLUSTER);
 DEFINE_STAT(STAT_Navigation_DetourPERM_TILES);
+DEFINE_STAT(STAT_Navigation_DetourPERM_TILE_LINK_BUILDER);
 
 DEFINE_STAT(STAT_DetourTileMemory);
 DEFINE_STAT(STAT_DetourTileMeshHeaderMemory);
diff --git a/Engine/Source/Runtime/NavigationSystem/Public/NavMesh/LinkGenerationDebugFlags.h b/Engine/Source/Runtime/NavigationSystem/Public/NavMesh/LinkGenerationDebugFlags.h
new file mode 100644
--- /dev/null
+++ b/Engine/Source/Runtime/NavigationSystem/Public/NavMesh/LinkGenerationDebugFlags.h
@@ -0,0 +1,17 @@
+﻿// Copyright Epic Games, Inc. All Rights Reserved.
+#pragma once
+
+#include "LinkGenerationDebugFlags.generated.h"
+
+UENUM(meta = (Bitflags, UseEnumValuesAsMaskValuesInEditor = "true"))
+enum class ELinkGenerationDebugFlags : uint16
+{
+	WalkableSurface				= 1 << 0,
+	WalkableBorders				= 1 << 1,
+	SelectedEdge				= 1 << 2,
+	SelectedEdgeTrajectory		= 1 << 3,
+	SelectedEdgeLandingSamples	= 1 << 4,
+	SelectedEdgeCollisions		= 1 << 5,
+	Links						= 1 << 6,
+};
+ENUM_CLASS_FLAGS(ELinkGenerationDebugFlags);
diff --git a/Engine/Source/Runtime/NavigationSystem/Public/NavMesh/RecastNavMesh.h b/Engine/Source/Runtime/NavigationSystem/Public/NavMesh/RecastNavMesh.h
--- a/Engine/Source/Runtime/NavigationSystem/Public/NavMesh/RecastNavMesh.h
+++ b/Engine/Source/Runtime/NavigationSystem/Public/NavMesh/RecastNavMesh.h
@@ -520,6 +520,13 @@ struct FRecastNavMeshTileGenerationDebug
 
 	UPROPERTY(EditAnywhere, Category = Debug)
 	uint32 bTileCacheDetailMesh : 1;
+
+	UPROPERTY(EditAnywhere, Category = Debug, meta = (Bitmask, BitmaskEnum = "/Script/NavigationSystem.ELinkGenerationDebugFlags"))
+	uint16 LinkGenerationDebugFlags;
+
+	/** Using -1 as no selected edge. */
+	UPROPERTY(EditAnywhere, Category = Debug)
+	int32 LinkGenerationSelectedEdge;
 };
 
 /**
@@ -886,6 +893,10 @@ class ARecastNavMesh : public ANavigationData
 	/* In a world partitioned map, is this navmesh using world partitioning */
 	UPROPERTY(EditAnywhere, Category=Generation, config, meta = (EditCondition = "bAllowWorldPartitionedNavMesh", HideEditConditionToggle, DisplayName = "IsWorldPartitionedNavMesh"))
 	uint32 bIsWorldPartitioned : 1;
+
+	/** Experimental: if set, navlinks will be automatically generated. */ 
+	UPROPERTY(EditAnywhere, Category=Generation, config)
+	uint32 bGenerateNavLinks : 1;
 	
 	/** controls whether voxel filtering will be applied (via FRecastTileGenerator::ApplyVoxelFilter). 
 	 *	Results in generated navmesh better fitting navigation bounds, but hits (a bit) generation performance */
diff --git a/Engine/Source/Runtime/NavigationSystem/Public/NavMesh/RecastNavMeshGenerator.h b/Engine/Source/Runtime/NavigationSystem/Public/NavMesh/RecastNavMeshGenerator.h
--- a/Engine/Source/Runtime/NavigationSystem/Public/NavMesh/RecastNavMeshGenerator.h
+++ b/Engine/Source/Runtime/NavigationSystem/Public/NavMesh/RecastNavMeshGenerator.h
@@ -40,6 +40,7 @@ struct FKAggregateGeom;
 struct FTileCacheCompressor;
 struct FTileCacheAllocator;
 struct FTileGenerationContext;
+struct dtLinkBuilderData;
 class dtNavMesh;
 class FNavRegenTimeSliceManager;
 class UNavigationSystemV1;
@@ -64,6 +65,8 @@ struct FRecastBuildConfig : public rcConfig
 	uint32 bFilterLowSpanSequences : 1;
 	/** if set, only low height spans with corresponding area modifier will be stored in tile cache (reduces memory, can't modify without full tile rebuild) */
 	uint32 bFilterLowSpanFromTileCache : 1;
+	/** if set, navigation links will be automatically generated */
+	uint32 bGenerateLinks : 1;
 
 	/** region partitioning method used by tile cache */
 	int32 TileCachePartitionType;
@@ -87,7 +90,7 @@ struct FRecastBuildConfig : public rcConfig
 	/** Ledge filtering mode */
 	ENavigationLedgeSlopeFilterMode LedgeSlopeFilterMode;
 	/** Is the config completely setup */
-	bool bIsTileSetupConfigCompleted = false;
+	bool bIsTileSetupConfigCompleted;
 	
 	FRecastBuildConfig()
 	{
@@ -104,6 +107,7 @@ struct FRecastBuildConfig : public rcConfig
 		bUseExtraTopCellWhenMarkingAreas = true;
 		bFilterLowSpanSequences = false;
 		bFilterLowSpanFromTileCache = false;
+		bGenerateLinks = false;
 		// Still initializing, even though the property is deprecated, to avoid static analysis warnings
 		MaxPolysPerTile = -1;
 		AgentIndex = 0;
@@ -438,20 +442,43 @@ protected:
 	NAVIGATIONSYSTEM_API bool RecastBuildTileCache(FNavMeshBuildContext& BuildContext, FTileRasterizationContext& RasterContext);
 	/** End functions used by GenerateCompressedLayersTimeSliced / GenerateCompressedLayers */
 
-	/** builds CompressedLayers array (geometry + modifiers) time sliced*/
+	/** Builds CompressedLayers array (geometry + modifiers) time sliced*/
 	NAVIGATIONSYSTEM_API virtual ETimeSliceWorkResult GenerateCompressedLayersTimeSliced(FNavMeshBuildContext& BuildContext);
+
+	
 	/** builds CompressedLayers array (geometry + modifiers) */
+	UE_DEPRECATED(5.5, "Use the overload with dtLinkBuilderData instead.")
 	NAVIGATIONSYSTEM_API virtual bool GenerateCompressedLayers(FNavMeshBuildContext& BuildContext);
 
 	/** Builds a navigation data layer */
+	UE_DEPRECATED(5.5, "Use the overload with dtLinkBuilderData instead.")
 	NAVIGATIONSYSTEM_API bool GenerateNavigationDataLayer(FNavMeshBuildContext& BuildContext, FTileCacheCompressor& TileCompressor, FTileCacheAllocator& GenNavAllocator, FTileGenerationContext& GenerationContext, int32 LayerIdx);
 
 	/** builds NavigationData array (layers + obstacles) time sliced */
+	UE_DEPRECATED(5.5, "Use the overload with dtLinkBuilderData instead.")
 	NAVIGATIONSYSTEM_API ETimeSliceWorkResult GenerateNavigationDataTimeSliced(FNavMeshBuildContext& BuildContext);
 
 	/** builds NavigationData array (layers + obstacles) */
+	UE_DEPRECATED(5.5, "Use the overload with dtLinkBuilderData instead.")
 	NAVIGATIONSYSTEM_API bool GenerateNavigationData(FNavMeshBuildContext& BuildContext);
+	
+	
+	/** Builds CompressedLayers array (geometry + modifiers) */
+	NAVIGATIONSYSTEM_API virtual bool GenerateCompressedLayers(FNavMeshBuildContext& BuildContext, dtLinkBuilderData& OutLinkBuilderData);
 
+	/** Builds a navigation data layer */
+	NAVIGATIONSYSTEM_API bool GenerateNavigationDataLayer(FNavMeshBuildContext& BuildContext, FTileCacheCompressor& TileCompressor, FTileCacheAllocator& GenNavAllocator, FTileGenerationContext& GenerationContext, dtLinkBuilderData& InOutLinkBuilderData, int32 LayerIdx);
+
+	/** Builds NavigationData array (layers + obstacles) time sliced */
+	NAVIGATIONSYSTEM_API ETimeSliceWorkResult GenerateNavigationDataTimeSliced(FNavMeshBuildContext& BuildContext, dtLinkBuilderData& InOutLinkBuilderData);
+
+	/** Builds NavigationData array (layers + obstacles) */
+	NAVIGATIONSYSTEM_API bool GenerateNavigationData(FNavMeshBuildContext& BuildContext, dtLinkBuilderData& LinkBuilderData);
+
+	/** Builds navigation links */
+	dtStatus BuildTileCacheLinks(FNavMeshBuildContext& BuildContext, struct dtTileCacheAlloc* alloc, const dtTileCacheLayer& layer,
+		const struct dtTileCacheContourSet& lcset, const dtLinkBuilderData& linkBuilderData, TArray<FNavigationLink>& OutGeneratedLinks) const;
+	
 	NAVIGATIONSYSTEM_API virtual void ApplyVoxelFilter(struct rcHeightfield* SolidHF, FVector::FReal WalkableRadius);
 
 	/** Compute rasterization mask */
diff --git a/Engine/Source/Runtime/Navmesh/Private/DebugUtils/DebugDraw.cpp b/Engine/Source/Runtime/Navmesh/Private/DebugUtils/DebugDraw.cpp
--- a/Engine/Source/Runtime/Navmesh/Private/DebugUtils/DebugDraw.cpp
+++ b/Engine/Source/Runtime/Navmesh/Private/DebugUtils/DebugDraw.cpp
@@ -53,6 +53,33 @@ void duIntToCol(int i, float* col)
 	col[2] = 1 - b*63.0f/255.0f;
 }
 
+//@UE BEGIN
+// Color palette
+unsigned int duColor::red			= duRGBA(0xe53935ff);		
+unsigned int duColor::pink			= duRGBA(0xd81b60ff);
+unsigned int duColor::purple		= duRGBA(0x8e24aaff);
+unsigned int duColor::deepPurple	= duRGBA(0x5e35b1ff);
+unsigned int duColor::indigo		= duRGBA(0x3949abff);
+unsigned int duColor::blue			= duRGBA(0x1e88e5ff);
+unsigned int duColor::lightBlue		= duRGBA(0x039be5ff);
+unsigned int duColor::cyan			= duRGBA(0x00acc1ff);
+unsigned int duColor::teal			= duRGBA(0x00897bff);
+unsigned int duColor::green			= duRGBA(0x43a047ff);
+unsigned int duColor::lightGreen	= duRGBA(0x7cb342ff);
+unsigned int duColor::lime			= duRGBA(0xc0ca33ff);
+unsigned int duColor::yellow		= duRGBA(0xfdd835ff);
+unsigned int duColor::amber			= duRGBA(0xffb300ff);
+unsigned int duColor::orange		= duRGBA(0xfb8c00ff);
+unsigned int duColor::orangeRed		= duRGBA(0xf4511eff);
+unsigned int duColor::brown			= duRGBA(0x6d4c41ff);
+
+unsigned int duColor::white			= duRGBA(0xffffffff);
+unsigned int duColor::lightGrey		= duRGBA(0xe0e0e0ff);
+unsigned int duColor::grey			= duRGBA(0x9e9e9eff);
+unsigned int duColor::drakGrey		= duRGBA(0x616161ff);
+unsigned int duColor::black			= duRGBA(0x000000ff);
+//@UE END
+
 void duCalcBoxColors(unsigned int* colors, unsigned int colTop, unsigned int colSide)
 {
 	if (!colors) return;
diff --git a/Engine/Source/Runtime/Navmesh/Private/DebugUtils/DetourNavLinkDebugDraw.cpp b/Engine/Source/Runtime/Navmesh/Private/DebugUtils/DetourNavLinkDebugDraw.cpp
new file mode 100644
--- /dev/null
+++ b/Engine/Source/Runtime/Navmesh/Private/DebugUtils/DetourNavLinkDebugDraw.cpp
@@ -0,0 +1,561 @@
+// Copyright Epic Games, Inc. All Rights Reserved.
+#include "DebugUtils/DetourNavLinkDebugDraw.h"
+
+#include "DebugUtils/DebugDraw.h"
+#include "DebugUtils/RecastDebugDraw.h"
+#include "Detour/DetourCommon.h"
+#include "Detour/DetourNavLinkBuilder.h"
+#include "Recast/Recast.h"
+
+//@UE BEGIN
+namespace UE::Detour::Private
+{
+	dtReal getPathLen(const dtReal* path, int npath)
+	{
+		if (npath < 2)
+			return 0;
+		
+		dtReal totd = 0;
+		for (int i = 0; i < npath-1; ++i)
+		{
+			const dtReal* sp = &path[i*3];
+			const dtReal* sq = &path[(i+1)*3];
+			totd += sqrtf(dtVdistSqr(sp,sq));
+		}
+		
+		return totd;
+	}
+
+	void getPointAlongPath(dtReal dist, const dtReal* path, int npath, dtReal* pt, dtReal* dir)
+	{
+		if (npath == 0)
+			return;
+		
+		if (npath == 1)
+		{
+			if (dir)
+				dtVset(dir,1,0,0);
+			dtVcopy(pt, path);
+			return;
+		}
+		
+		if (dist <= 0)
+		{
+			if (dir)
+				dtVsub(dir, &path[1*3], &path[0*3]);
+			
+			dtVcopy(pt, path);
+			return;
+		}
+		
+		float totd = 0;
+		for (int i = 0; i < npath-1; ++i)
+		{
+			const dtReal* sp = &path[i*3];
+			const dtReal* sq = &path[(i+1)*3];
+			const float sd = sqrtf(dtVdistSqr(sp,sq));
+			if (dist >= totd && dist <= totd+sd)
+			{
+				dtVlerp(pt, sp,sq,(dist-totd)/sd);
+				if (dir)
+					dtVsub(dir, sq,sp);
+				return;
+			}
+			totd += sd;
+		}
+		
+		dtVcopy(pt, &path[(npath-1)*3]);
+		if (dir)
+		{
+			dtVsub(dir, &path[(npath-1)*3], &path[(npath-2)*3]);
+		}
+	}
+
+	void drawTrajectory(duDebugDraw* dd, const dtReal* pa, const dtReal* pb, const dtNavLinkBuilder::Trajectory2D* trajectory, const unsigned int color)
+	{
+		const float startx = trajectory->spine[0];
+		const float endx = trajectory->spine[(trajectory->nspine-1)*2];
+		const float deltax = endx - startx;
+		
+		const float starty = trajectory->spine[1];
+		const float endy = trajectory->spine[(trajectory->nspine-1)*2+1];
+		
+		dtReal pts[3*dtNavLinkBuilder::MAX_SPINE];
+		int npts = trajectory->nspine;
+		for (int i = 0; i < trajectory->nspine; ++i)
+		{
+			const float* spt = &trajectory->spine[i*2];
+			const float u = (spt[0] - startx)/deltax;
+			const float dy = spt[1] - dtLerp(starty, endy, u);
+			dtReal* p = &pts[i*3];
+			dtVlerp(p, pa, pb, u);
+			p[1] += dy;
+		}
+			
+		const float len = getPathLen(pts, npts);
+		const int nsegs = ceilf(len / 0.3f);
+		
+		for (int i = 0; i < nsegs*2; ++i)
+		{
+			float u = (float)i / (float)(nsegs*2);
+			dtReal pt[3];
+			getPointAlongPath(u*len, pts, npts, pt, nullptr);
+			dd->vertex(pt, color);
+		}
+	}
+
+	void drawTrajectorySlice(duDebugDraw* dd, const dtReal* pa, const dtReal* pb, const dtNavLinkBuilder::Trajectory2D* trajectory, const unsigned int color)
+	{
+		unsigned int colt = duTransCol(color, 255);
+		unsigned int colb = duTransCol(duLerpCol(color,duColor::black,96), 255);
+		
+		dd->begin(DU_DRAW_QUADS);
+		dtReal p0[3], p1[3], p2[3], p3[3];
+		for (int i = 0; i < trajectory->nsamples; ++i)
+		{
+			const dtNavLinkBuilder::TrajectorySample* s = &trajectory->samples[i];
+			const dtReal u = (float)i / (float)(trajectory->nsamples-1);
+			dtVlerp(p1, pa, pb, u);
+			dtVcopy(p0, p1);
+			p0[1] += s->ymin;
+			p1[1] += s->ymax;
+			
+			if (i > 0)
+			{
+				dd->vertex(p0, colb);
+				dd->vertex(p1, colt);
+				dd->vertex(p3, colt);
+				dd->vertex(p2, colb);
+
+				dd->vertex(p2, colb);
+				dd->vertex(p3, colt);
+				dd->vertex(p1, colt);
+				dd->vertex(p0, colb);
+			}
+			
+			dtVcopy(p2, p0);
+			dtVcopy(p3, p1);
+		}
+		dd->end();
+
+		colb = duDarkenCol(colb);
+		colt = duDarkenCol(colt);
+		
+		dd->begin(DU_DRAW_LINES, 2.0f);
+		for (int i = 0; i < trajectory->nsamples; ++i)
+		{
+			const dtNavLinkBuilder::TrajectorySample* s = &trajectory->samples[i];
+			const float u = (float)i / (float)(trajectory->nsamples-1);
+			dtVlerp(p1, pa, pb, u);
+			dtVcopy(p0, p1);
+			p0[1] += s->ymin;
+			p1[1] += s->ymax;
+			
+			if (i == 0 || i == (trajectory->nsamples-1))
+			{
+				dd->vertex(p0, colb);
+				dd->vertex(p1, colt);
+			}
+			if (i > 0)
+			{
+				dd->vertex(p0, colb);
+				dd->vertex(p2, colb);
+				dd->vertex(p1, colt);
+				dd->vertex(p3, colt);
+			}
+			
+			dtVcopy(p2, p0);
+			dtVcopy(p3, p1);
+		}
+		dd->end();
+	}
+} // UE::Detour::Private
+
+void duDebugDrawNavLinkBuilder(duDebugDraw* dd, const dtNavLinkBuilder& linkBuilder, unsigned int drawFlags, const dtNavLinkBuilder::EdgeSampler* es)
+{
+	if (dd == nullptr)
+		return;
+	
+	if (drawFlags & DRAW_WALKABLE_SURFACE)
+	{
+		if (linkBuilder.m_chf)
+		{
+			duDebugDrawCompactHeightfieldSolid(dd, *linkBuilder.m_chf);	
+		}
+	}
+
+	if (drawFlags & DRAW_BORDERS)
+	{
+		const dtNavLinkBuilder::Edge* edges = linkBuilder.m_edges;
+		const int nedges = linkBuilder.m_nedges;
+		const int selectedEdge = linkBuilder.m_debugSelectedEdge;
+		
+		if (nedges)
+		{
+			dd->begin(DU_DRAW_LINES, 3.0f);
+			for (int i = 0; i < nedges; ++i)
+			{
+				unsigned int col = duRGBA(0,96,128,255);
+				if (i == selectedEdge)
+					continue;
+				
+				dd->vertex(edges[i].sp, col);
+				dd->vertex(edges[i].sq, col);
+			}
+			dd->end();
+			
+			dd->begin(DU_DRAW_POINTS, 8.0f);
+			for (int i = 0; i < nedges; ++i)
+			{
+				unsigned int col = duRGBA(0,96,128,255);
+				if (i == selectedEdge)
+					continue;
+				
+				dd->vertex(edges[i].sp, col);
+				dd->vertex(edges[i].sq, col);
+			}
+			dd->end();
+			
+			if (selectedEdge >= 0 && selectedEdge < nedges)
+			{
+				unsigned int col = duRGBA(68,36,36,255);
+				dd->begin(DU_DRAW_LINES, 3.0f);
+				dd->vertex(edges[selectedEdge].sp, col);
+				dd->vertex(edges[selectedEdge].sq, col);
+				dd->end();
+				dd->begin(DU_DRAW_POINTS, 8.0f);
+				dd->vertex(edges[selectedEdge].sp, col);
+				dd->vertex(edges[selectedEdge].sq, col);
+				dd->end();
+			}
+
+			dd->begin(DU_DRAW_POINTS, 4.0f);
+			for (int i = 0; i < nedges; ++i)
+			{
+				unsigned int col = duColor::lightGrey;
+				dd->vertex(edges[i].sp, col);
+				dd->vertex(edges[i].sq, col);
+			}
+			dd->end();
+		}
+	}
+	
+	if (drawFlags & DRAW_LINKS)
+	{
+		const dtNavLinkBuilder::JumpLink* links = linkBuilder.m_links;
+
+		unsigned int jumpDownCol0 = duLerpCol(duColor::blue, duColor::white, 200);
+		unsigned int jumpDownCol1 = duColor::blue;
+		unsigned int jumpOverCol0 = duLerpCol(duColor::lightGrey, duColor::white, 200);
+		unsigned int jumpOverCol1 = duColor::lightGrey;
+
+		auto selectColors = [&](dtNavLinkAction action, unsigned int& col0, unsigned int& col1)
+		{
+			if (action == DT_LINK_ACTION_JUMP_DOWN)
+			{
+				col0 = jumpDownCol0;
+				col1 = jumpDownCol1;
+			}
+			else if (action == DT_LINK_ACTION_JUMP_OVER)
+			{
+				col0 = jumpOverCol0;
+				col1 = jumpOverCol1;
+			}
+		};
+
+		unsigned int col0 = duColor::black;
+		unsigned int col1 = duColor::black;
+		
+		const int nlinks = linkBuilder.m_nlinks;
+		if (nlinks)
+		{
+			dd->begin(DU_DRAW_QUADS);
+			for (int i = 0; i < nlinks; ++i)
+			{
+				const dtNavLinkBuilder::JumpLink* link = &links[i];
+				
+				if (link->flags == dtNavLinkBuilder::INVALID)
+					continue;
+
+				selectColors(link->action, col0, col1);
+				
+				for (int j = 0; j < link->nspine-1; ++j)
+				{
+					int u = (j*255)/link->nspine;
+					unsigned int col = duTransCol(duLerpCol(col0,col1,u),128);
+					if (link->flags == dtNavLinkBuilder::INVALID)
+						col = duRGBA(255,0,0,64);
+					
+					dd->vertex(&link->spine1[j*3], col);
+					dd->vertex(&link->spine1[(j+1)*3], col);
+					dd->vertex(&link->spine0[(j+1)*3], col);
+					dd->vertex(&link->spine0[j*3], col);
+				}
+			}
+			dd->end();
+			
+			dd->begin(DU_DRAW_LINES, 3.0f);
+			for (int i = 0; i < nlinks; ++i)
+			{
+				const dtNavLinkBuilder::JumpLink* link = &links[i];
+				if (link->flags == dtNavLinkBuilder::INVALID)
+					continue;
+
+				selectColors(link->action, col0, col1);
+				
+				for (int j = 0; j < link->nspine-1; ++j)
+				{
+					unsigned int col = duTransCol(duDarkenCol(col1),128);
+					
+					dd->vertex(&link->spine0[j*3], col);
+					dd->vertex(&link->spine0[(j+1)*3], col);
+					dd->vertex(&link->spine1[j*3], col);
+					dd->vertex(&link->spine1[(j+1)*3], col);
+				}
+
+				dd->vertex(&link->spine0[0], duDarkenCol(col1));
+				dd->vertex(&link->spine1[0], duDarkenCol(col1));
+
+				dd->vertex(&link->spine0[(link->nspine-1)*3], duDarkenCol(col1));
+				dd->vertex(&link->spine1[(link->nspine-1)*3], duDarkenCol(col1));
+			}
+			dd->end();
+
+			dd->begin(DU_DRAW_POINTS, 8.0f);
+			for (int i = 0; i < nlinks; ++i)
+			{
+				const dtNavLinkBuilder::JumpLink* link = &links[i];
+				if (link->flags == dtNavLinkBuilder::INVALID)
+					continue;
+
+				selectColors(link->action, col0, col1);
+				
+				dd->vertex(&link->spine0[0], duDarkenCol(col1));
+				dd->vertex(&link->spine1[0], duDarkenCol(col1));
+				dd->vertex(&link->spine0[(link->nspine-1)*3], duDarkenCol(col1));
+				dd->vertex(&link->spine1[(link->nspine-1)*3], duDarkenCol(col1));
+			}
+			dd->end();
+			
+			dd->begin(DU_DRAW_POINTS, 4.0f);
+			for (int i = 0; i < nlinks; ++i)
+			{
+				const dtNavLinkBuilder::JumpLink* link = &links[i];
+				if (link->flags == dtNavLinkBuilder::INVALID)
+					continue;
+				dd->vertex(&link->spine0[0], duColor::lightGrey);
+				dd->vertex(&link->spine1[0], duColor::lightGrey);
+				dd->vertex(&link->spine0[(link->nspine-1)*3], duColor::lightGrey);
+				dd->vertex(&link->spine1[(link->nspine-1)*3], duColor::lightGrey);
+			}
+			dd->end();
+		}
+	}
+
+	if (drawFlags & DRAW_SELECTED_EDGE)
+	{
+		if (es)
+		{
+			dd->begin(DU_DRAW_LINES, 2.0f);
+
+			// Axis
+			constexpr float d = 20.f;
+			dd->vertex(es->rigp, duColor::red);
+			dd->vertex(es->rigp[0]+es->ax[0]*d, es->rigp[1]+es->ax[1]*d, es->rigp[2]+es->ax[2]*d, duColor::red);
+
+			dd->vertex(es->rigp, duColor::green);
+			dd->vertex(es->rigp[0]+es->ay[0]*d, es->rigp[1]+es->ay[1]*d, es->rigp[2]+es->ay[2]*d, duColor::green);
+
+			dd->vertex(es->rigp, duColor::blue);
+			dd->vertex(es->rigp[0]+es->az[0]*d, es->rigp[1]+es->az[1]*d, es->rigp[2]+es->az[2]*d, duColor::blue);
+			
+			dd->end();
+			
+			if (drawFlags & DRAW_TRAJECTORY)
+			{
+				const float r = es->groundRange;
+
+				unsigned int col = duLerpCol(duRGBA(255,192,0,255), duColor::white, 64);
+				unsigned int cola = duTransCol(col,192);
+				unsigned int colb = duColor::white;
+				
+				// Start segment.
+				dd->begin(DU_DRAW_LINES, 3.0f);
+				dd->vertex(es->start.p[0],es->start.p[1],es->start.p[2], col);
+				dd->vertex(es->start.q[0],es->start.q[1],es->start.q[2], col);
+				dd->end();
+				
+				dd->begin(DU_DRAW_LINES, 1.0f);
+				dd->vertex(es->start.p[0],es->start.p[1]-r,es->start.p[2], colb);
+				dd->vertex(es->start.p[0],es->start.p[1]+r,es->start.p[2], colb);
+				dd->vertex(es->start.p[0],es->start.p[1]+r,es->start.p[2], colb);
+				dd->vertex(es->start.q[0],es->start.q[1]+r,es->start.q[2], colb);
+				dd->vertex(es->start.q[0],es->start.q[1]+r,es->start.q[2], colb);
+				dd->vertex(es->start.q[0],es->start.q[1]-r,es->start.q[2], colb);
+				dd->vertex(es->start.q[0],es->start.q[1]-r,es->start.q[2], colb);
+				dd->vertex(es->start.p[0],es->start.p[1]-r,es->start.p[2], colb);
+				dd->end();
+				
+				// End segment.
+				dd->begin(DU_DRAW_LINES, 3.0f);
+				dd->vertex(es->end.p[0],es->end.p[1],es->end.p[2], col);
+				dd->vertex(es->end.q[0],es->end.q[1],es->end.q[2], col);
+				dd->end();
+				
+				dd->begin(DU_DRAW_LINES, 1.0f);
+				dd->vertex(es->end.p[0],es->end.p[1]-r,es->end.p[2], colb);
+				dd->vertex(es->end.p[0],es->end.p[1]+r,es->end.p[2], colb);
+				dd->vertex(es->end.p[0],es->end.p[1]+r,es->end.p[2], colb);
+				dd->vertex(es->end.q[0],es->end.q[1]+r,es->end.q[2], colb);
+				dd->vertex(es->end.q[0],es->end.q[1]+r,es->end.q[2], colb);
+				dd->vertex(es->end.q[0],es->end.q[1]-r,es->end.q[2], colb);
+				dd->vertex(es->end.q[0],es->end.q[1]-r,es->end.q[2], colb);
+				dd->vertex(es->end.p[0],es->end.p[1]-r,es->end.p[2], colb);
+				dd->end();
+		 
+				dd->begin(DU_DRAW_LINES, 4.0f);
+				UE::Detour::Private::drawTrajectory(dd, es->start.p, es->end.p, &es->trajectory, cola);
+				UE::Detour::Private::drawTrajectory(dd, es->start.q, es->end.q, &es->trajectory, cola);
+				dd->end();
+
+
+				dd->begin(DU_DRAW_POINTS, 8.0f);
+				dd->vertex(es->start.p[0],es->start.p[1],es->start.p[2], duDarkenCol(col));
+				dd->vertex(es->start.q[0],es->start.q[1],es->start.q[2], duDarkenCol(col));
+				dd->vertex(es->end.p[0],es->end.p[1],es->end.p[2], duDarkenCol(col));
+				dd->vertex(es->end.q[0],es->end.q[1],es->end.q[2], duDarkenCol(col));
+				dd->end();
+				
+				unsigned int colm = duColor::white;
+				dd->begin(DU_DRAW_POINTS, 3.0f);
+				dd->vertex(es->start.p[0],es->start.p[1],es->start.p[2], colm);
+				dd->vertex(es->start.q[0],es->start.q[1],es->start.q[2], colm);
+				dd->vertex(es->end.p[0],es->end.p[1],es->end.p[2], colm);
+				dd->vertex(es->end.q[0],es->end.q[1],es->end.q[2], colm);
+				dd->end();
+			}
+			
+			if (drawFlags & DRAW_LAND_SAMPLES)
+			{
+				if (es->start.ngsamples)
+				{
+					dd->begin(DU_DRAW_POINTS, 8.0f);
+					for (int i = 0; i < es->start.ngsamples; ++i)
+					{
+						const dtNavLinkBuilder::GroundSample* s = &es->start.gsamples[i];
+						const float u = (float)i/(float)(es->start.ngsamples-1);
+						dtReal spt[3], ept[3];
+						dtVlerp(spt, es->start.p, es->start.q, u);
+						dtVlerp(ept, es->start.p, es->start.q, u);
+						
+						unsigned int col = duRGBA(48,16,16,255);
+						dtReal off = 1;
+						if ((s->flags & dtNavLinkBuilder::HAS_GROUND) == 0)
+						{
+							off = 0;
+							col = duRGBA(220,32,32,255);
+						}
+						
+						spt[1] = s->height + off;
+						dd->vertex(spt, col);
+					}
+					dd->end();
+					
+					dd->begin(DU_DRAW_POINTS, 4.0f);
+					for (int i = 0; i < es->start.ngsamples; ++i)
+					{
+						const dtNavLinkBuilder::GroundSample* s = &es->start.gsamples[i];
+						const float u = (float)i/(float)(es->start.ngsamples-1);
+						dtReal spt[3], ept[3];
+						dtVlerp(spt, es->start.p, es->start.q,u);
+						dtVlerp(ept, es->start.p, es->start.q,u);
+						float off = 0;
+						if (s->flags & dtNavLinkBuilder::HAS_GROUND)
+						{
+							off = 1.f;
+						}
+						spt[1] = s->height + off;
+						dd->vertex(spt, duColor::white);
+					}
+					dd->end();
+					
+				}
+
+				if (es->end.ngsamples)
+				{
+					dd->begin(DU_DRAW_POINTS, 8.0f);
+					for (int i = 0; i < es->end.ngsamples; ++i)
+					{
+						const dtNavLinkBuilder::GroundSample* s = &es->end.gsamples[i];
+						const float u = (float)i/(float)(es->end.ngsamples-1);
+						dtReal spt[3], ept[3];
+						dtVlerp(spt, es->end.p, es->end.q, u);
+						dtVlerp(ept, es->end.p, es->end.q, u);
+						
+						unsigned int col = duRGBA(48,16,16,255);
+						float off = 1.f;
+						if ((s->flags & dtNavLinkBuilder::HAS_GROUND) == 0)
+						{
+							off = 0;
+							col = duRGBA(220,32,32,255);
+						}
+						
+						spt[1] = s->height + off;
+						dd->vertex(spt, col);
+					}
+					dd->end();
+					
+					dd->begin(DU_DRAW_POINTS, 4.0f);
+					for (int i = 0; i < es->end.ngsamples; ++i)
+					{
+						const dtNavLinkBuilder::GroundSample* s = &es->end.gsamples[i];
+						const float u = (float)i/(float)(es->end.ngsamples-1);
+						dtReal spt[3], ept[3];
+						dtVlerp(spt, es->end.p, es->end.q, u);
+						dtVlerp(ept, es->end.p, es->end.q, u);
+						
+						float off = 0;
+						if (s->flags & dtNavLinkBuilder::HAS_GROUND)
+						{
+							off = 1.f;
+						}
+						spt[1] = s->height + off;
+						dd->vertex(spt, duColor::white);
+					}
+					dd->end();
+				}
+			}
+			
+			if (drawFlags & DRAW_COLLISION_SLICES)
+			{
+				if (es->start.ngsamples && es->start.ngsamples == es->end.ngsamples)
+				{
+					const int nsamples = es->start.ngsamples;
+					dd->depthMask(true);
+					for (int i = 0; i < nsamples; ++i)
+					{
+						const dtNavLinkBuilder::GroundSample* ssmp = &es->start.gsamples[i];
+						const dtNavLinkBuilder::GroundSample* esmp = &es->end.gsamples[i];
+						if ((ssmp->flags & dtNavLinkBuilder::HAS_GROUND) == 0 || (esmp->flags & dtNavLinkBuilder::HAS_GROUND) == 0)
+							continue;
+						
+						const float u = (float)i/(float)(nsamples-1);
+						dtReal spt[3], ept[3];
+						dtVlerp(spt, es->start.p, es->start.q, u);
+						dtVlerp(ept, es->end.p, es->end.q, u);
+
+						spt[1] = ssmp->height;
+						ept[1] = esmp->height;
+						
+						if (ssmp->flags & dtNavLinkBuilder::UNRESTRICTED)
+							UE::Detour::Private::drawTrajectorySlice(dd, spt, ept, &es->trajectory, duRGBA(32,255,96,255));
+						else
+							UE::Detour::Private::drawTrajectorySlice(dd, spt, ept, &es->trajectory, duRGBA(255,32,32,255));
+					}
+					dd->depthMask(true);
+				}
+			}
+		}
+	}
+}
+//@UE END
diff --git a/Engine/Source/Runtime/Navmesh/Private/Detour/DetourNavLinkBuilder.cpp b/Engine/Source/Runtime/Navmesh/Private/Detour/DetourNavLinkBuilder.cpp
new file mode 100644
--- /dev/null
+++ b/Engine/Source/Runtime/Navmesh/Private/Detour/DetourNavLinkBuilder.cpp
@@ -0,0 +1,1034 @@
+// Copyright Epic Games, Inc. All Rights Reserved.
+#include "Detour/DetourNavLinkBuilder.h"
+
+#include "Detour/DetourAlloc.h"
+#include "Detour/DetourCommon.h"
+#include "DetourTileCache/DetourTileCacheBuilder.h"
+#include "Recast/Recast.h"
+
+//@UE BEGIN
+namespace UE::Detour::NavLink::Private
+{
+	static void insertSort(unsigned char* a, const int n)
+	{
+		int j;
+		for (int i = 1; i < n; i++)
+		{
+			const unsigned char value = a[i];
+			for (j = i - 1; j >= 0 && a[j] > value; j--)
+			{
+				a[j+1] = a[j];
+			}
+			a[j+1] = value;
+		}
+	}
+
+	static dtReal getClosestPtPtSeg(const dtReal* pt, const dtReal* sp, const dtReal* sq)
+	{
+		dtReal dir[3], diff[3];
+		dtVsub(dir, sq, sp);
+		dtVsub(diff, pt, sp);
+		dtReal t = dtVdot(dir,diff);
+		if (t <= 0.0)
+			return 0.0;
+	
+		const dtReal d = dtVdot(dir,dir);
+		if (t >= d)
+			return 1.0;
+
+		return t/d;
+	}
+
+	static void closestPtSegSeg(const dtReal* ap, const dtReal* aq, const dtReal* bp, const dtReal* bq, dtReal& s, dtReal& t)
+	{
+		constexpr dtReal EPSILON = 1e-6;
+		dtReal d1[3], d2[3], r[3];
+		dtVsub(d1, aq, ap);
+		dtVsub(d2, bq, bp);
+		dtVsub(r, ap, bp);
+		const dtReal a = dtVdot(d1,d1); // Squared length of segment S1, always nonnegative
+		const dtReal e = dtVdot(d2,d2); // Squared length of segment S2, always nonnegative
+		const dtReal f = dtVdot(d2,r);
+	
+		// Check if either or both segments degenerate into points
+		if (a <= EPSILON && e <= EPSILON)
+		{
+			// Both segments degenerate into points
+			s = t = 0.0;
+			return;
+		}
+	
+		if (a <= EPSILON)
+		{
+			s = 0.0;
+			t = f / e;
+			t = dtClamp(t, 0.0, 1.0);
+		}
+		else
+		{
+			const dtReal c = dtVdot(d1,r);
+			if (e <= EPSILON)
+			{
+				// Second segment degenerates into a point
+				t = 0.0;
+				s = dtClamp(-c / a, 0.0, 1.0); // t = 0 => s = (b*t - c) / a = -c / a
+			}
+			else
+			{
+				// The general nondegenerate case starts here
+				const dtReal b = dtVdot(d1,d2);
+				const dtReal denom = a*e-b*b; // Always nonnegative
+			
+				// If segments not parallel, compute closest point on L1 to L2, and
+				// clamp to segment S1. Else pick arbitrary s (here 0)
+				if (denom != 0.0)
+					s = dtClamp((b*f - c*e) / denom, 0.0, 1.0);
+				else
+					s = 0.0;
+			
+				// Compute point on L2 closest to S1(s) using
+				// t = Dot((p.start+D1*s)-q.start,D2) / Dot(D2,D2) = (b*s + f) / e
+				t = (b*s + f) / e;
+			
+				// If t in [0,1] done. Else clamp t, recompute s for the new value
+				// of t using s = Dot((q.start+D2*t)-p.start,D1) / Dot(D1,D1)= (t*b - c) / a
+				// and clamp s to [0, 1]
+				if (t < 0.0)
+				{
+					t = 0.0;
+					s = dtClamp(-c / a, 0.0, 1.0);
+				}
+				else if (t > 1.0)
+				{
+					t = 1.0;
+					s = dtClamp((b - c) / a, 0.0, 1.0);
+				}
+			}
+		}
+	}
+
+	static bool isectSegAABB(const dtReal* sp, const dtReal* sq,
+					  const float* amin, const float* amax,
+					  float& tmin, float& tmax)
+	{
+		static constexpr float EPS = 1e-6f;
+	
+		dtReal d[3];
+		dtVsub(d, sq, sp);
+		tmin = 0;  // set to -FLT_MAX to get first hit on the line
+		tmax = FLT_MAX;		// set to max distance ray can travel (for segment)
+	
+		// For all three slabs
+		for (int i = 0; i < 3; i++)
+		{
+			if (dtAbs(d[i]) < EPS)
+			{
+				// Ray is parallel to slab. No hit if origin not within slab
+				if (sp[i] < amin[i] || sp[i] > amax[i])
+					return false;
+			}
+			else
+			{
+				// Compute intersection t value of ray with near and far plane of slab
+				const float ood = 1.0f / d[i];
+				float t1 = (amin[i] - sp[i]) * ood;
+				float t2 = (amax[i] - sp[i]) * ood;
+
+				// Make t1 be intersection with near plane, t2 with far plane
+				if (t1 > t2)
+					dtSwap(t1, t2);
+
+				// Compute the intersection of slab intersections intervals
+				if (t1 > tmin)
+					tmin = t1;
+
+				if (t2 < tmax) 
+					tmax = t2;
+
+				// Exit with no collision as soon as slab intersection becomes empty
+				if (tmin > tmax)
+					return false;
+			}
+		}
+	
+		return true;
+	}
+
+	static dtReal distSegSegSqr(const dtReal* ap, const dtReal* aq, const dtReal* bp, const dtReal* bq)
+	{
+		dtReal s, t;
+		closestPtSegSeg(ap, aq, bp, bq, s,t);
+		dtReal anp[3], bnp[3];
+		dtVlerp(anp, ap,aq, s);
+		dtVlerp(bnp, bp,bq, t);
+		return dtVdistSqr(anp, bnp);
+	}
+	
+	static float getHeight(const float x, const float* pts, const int npts)
+	{
+		if (x <= pts[0])
+			return pts[1];
+
+		if (x >= pts[(npts-1)*2])
+			return pts[(npts-1)*2+1];
+	
+		for (int i = 1; i < npts; ++i)
+		{
+			const float* q = &pts[i*2];
+			if (x <= q[0])
+			{
+				const float* p = &pts[(i-1)*2];
+				const float u = (x-p[0]) / (q[0]-p[0]);
+				return dtLerp(p[1], q[1], u);
+			}
+		}
+	
+		return pts[(npts-1)*2+1];
+	}
+
+	inline bool overlapRange(const float amin, const float amax, const float bmin, const float bmax)
+	{
+		return (amin <= bmax && amax >= bmin);
+	}
+
+	inline void trans2d(dtReal* dst, const dtReal* ax, const dtReal* ay, const float* pt)
+	{
+		dst[0] = ax[0]*pt[0] + ay[0]*pt[1];
+		dst[1] = ax[1]*pt[0] + ay[1]*pt[1];
+		dst[2] = ax[2]*pt[0] + ay[2]*pt[1];
+	}
+}
+
+dtNavLinkBuilder::Trajectory2D::~Trajectory2D()
+{
+	dtFree(samples, DT_ALLOC_PERM_TILE_LINK_BUILDER);
+}
+
+dtNavLinkBuilder::GroundSegment::~GroundSegment()
+{
+	dtFree(gsamples, DT_ALLOC_PERM_TILE_LINK_BUILDER);
+}
+
+dtNavLinkBuilder::dtNavLinkBuilder() :
+	m_solid(nullptr),
+	m_chf(nullptr),
+	m_edges(nullptr),
+	m_nedges(0),
+	m_links(nullptr),
+	m_nlinks(0),
+	m_clinks(0),
+	m_debugSelectedEdge(-1)
+{
+}
+
+dtNavLinkBuilder::~dtNavLinkBuilder()
+{	
+	cleanup();
+}
+
+void dtNavLinkBuilder::cleanup()
+{
+	m_solid = nullptr;
+	m_chf = nullptr;
+
+	dtFree(m_edges, DT_ALLOC_PERM_TILE_LINK_BUILDER);
+	m_edges = nullptr;
+	m_nedges = 0;
+	
+	dtFree(m_links, DT_ALLOC_PERM_TILE_LINK_BUILDER);
+	m_links = nullptr;
+	m_nlinks = 0;
+	m_clinks = 0;
+}
+
+bool dtNavLinkBuilder::findEdges(rcContext& ctx, const rcConfig& cfg, const dtLinkBuilderConfig& builderConfig,
+                                 const dtTileCacheContourSet& lcset, const dtReal* orig, const dtLinkBuilderData& linkBuilderData)
+{
+	cleanup();
+	m_linkBuilderConfig = builderConfig;
+
+	m_cs = cfg.cs;
+	m_solid = linkBuilderData.solidHF;
+	m_chf = linkBuilderData.compactHF;
+
+	// Build edges.
+	m_nedges = 0;
+	for (int i = 0; i < lcset.nconts; ++i)
+	{
+		m_nedges += lcset.conts[i].nverts;
+	}
+
+	if (!m_nedges)
+	{
+		ctx.log(RC_LOG_ERROR, "fillEdges: No edges!");
+		return false;
+	}
+
+	dtAssert(m_edges == nullptr);
+	m_edges = (Edge*)dtAlloc(sizeof(Edge)*m_nedges, DT_ALLOC_PERM_TILE_LINK_BUILDER);
+	if (!m_edges)
+	{
+		ctx.log(RC_LOG_ERROR, "fillEdges: Failed to alloc edges (%d)", m_nedges);
+		return false;
+	}
+
+	const dtReal cs = cfg.cs;
+	const dtReal ch = cfg.ch;
+	
+	m_nedges = 0;
+	for (int i = 0; i < lcset.nconts; ++i)
+	{
+		const dtTileCacheContour& c = lcset.conts[i];
+		if (!c.nverts)
+			continue;
+
+		for (int j = 0, k = c.nverts-1; j < c.nverts; k=j++)
+		{
+			const unsigned short* va = &c.verts[k*4];
+			const unsigned short* vb = &c.verts[j*4];
+			
+			if ((va[3] & 0xf) != 0xf)	// A direction is set, so it's a portal edge.
+				continue;
+			
+			// Check k-j for matching contour
+			bool matchFound = false;
+			for (int ii = 0; ii < lcset.nconts; ++ii)
+			{
+				if (i == ii)
+					continue;
+				
+				const dtTileCacheContour& otherCont = lcset.conts[ii];
+				if (otherCont.nverts < 3)
+					continue;
+			
+				for (int jj = 0, kk = otherCont.nverts-1; jj < otherCont.nverts; kk=jj++)
+				{
+					const unsigned short* otherVa = &otherCont.verts[kk*4];
+					const unsigned short* otherVb = &otherCont.verts[jj*4];
+			
+					if( (dtVisEqual(va, otherVa) && dtVisEqual(vb, otherVb)) ||
+						(dtVisEqual(va, otherVb) && dtVisEqual(vb, otherVa))    )
+					{
+						// Same edge, skip it.
+						matchFound = true;
+						break;
+					}
+				}
+			
+				if (matchFound)
+				{
+					break;
+				}
+			}
+			
+			if (!matchFound)
+			{
+				// Add edge
+				Edge* e = &m_edges[m_nedges++];
+
+				e->sp[0] = orig[0] + vb[0]*cs;
+				e->sp[1] = orig[1] + (vb[1]+2)*ch;
+				e->sp[2] = orig[2] + vb[2]*cs;
+
+				e->sq[0] = orig[0] + va[0]*cs;
+				e->sq[1] = orig[1] + (va[1]+2)*ch;
+				e->sq[2] = orig[2] + va[2]*cs;
+			}
+		}
+	}
+
+	return true;
+}
+
+dtNavLinkBuilder::JumpLink* dtNavLinkBuilder::addLink()
+{
+	if (m_nlinks+1 > m_clinks)
+	{
+		m_clinks = m_clinks ? m_clinks*2 : 8;
+
+		JumpLink* n = (JumpLink*)dtAlloc(sizeof(JumpLink)*m_clinks, DT_ALLOC_PERM_TILE_LINK_BUILDER);
+
+		if (m_nlinks)
+		{
+			memcpy(n, m_links, sizeof(JumpLink)*m_nlinks);
+		}
+		dtFree(m_links, DT_ALLOC_PERM_TILE_LINK_BUILDER);
+		m_links = n;
+	}
+	
+	JumpLink* link = &m_links[m_nlinks++];
+	return link;
+}
+
+void dtNavLinkBuilder::addEdgeLinks(const dtLinkBuilderConfig& acfg, const EdgeSampler* es)
+{
+	using namespace UE::Detour::NavLink::Private;
+	
+	if (es->start.ngsamples != es->end.ngsamples)
+	{
+		return;
+	}
+	
+	const int nsamples = es->start.ngsamples;
+
+	// Filter small holes.
+	constexpr int RAD = 2;
+	GroundSampleFlag kernel[RAD*2+1];
+
+	GroundSampleFlag* nflags = (GroundSampleFlag*)dtAlloc(sizeof(GroundSampleFlag)*nsamples, DT_ALLOC_TEMP);
+	
+	for (int i = 0; i < nsamples; ++i)
+	{
+		const int a = dtMax(0, i-RAD);
+		const int b = dtMin(nsamples-1, i+RAD);
+		int nkernel = 0;
+		for (int j = a; j <= b; ++j)
+		{
+			kernel[nkernel++] = (GroundSampleFlag)(es->start.gsamples[i].flags & UNRESTRICTED);
+		}
+		insertSort((unsigned char*)kernel, nkernel);
+		nflags[i] = kernel[(nkernel+1)/2];
+	}
+
+	const int agentRadiusVx = acfg.agentRadius / acfg.cellSize;
+	
+	// Build segments
+	int start = -1;
+	for (int i = 0; i <= nsamples; ++i)
+	{
+		const bool valid = i < nsamples && nflags[i] != UNSET;
+		if (start == -1)
+		{
+			if (valid)
+				start = i;
+		}
+		else
+		{
+			if (!valid)
+			{
+				if (i - start > agentRadiusVx)
+				{
+					const float u0 = (float)start/(float)nsamples;
+					const float u1 = (float)i/(float)nsamples;
+
+					dtReal sp[3], sq[3], ep[3], eq[3];
+
+					dtVlerp(sp, es->start.p,es->start.q, u0);
+					dtVlerp(sq, es->start.p,es->start.q, u1);
+					dtVlerp(ep, es->end.p,es->end.q, u0);
+					dtVlerp(eq, es->end.p,es->end.q, u1);
+					sp[1] = es->start.gsamples[start].height;
+					sq[1] = es->start.gsamples[i-1].height;
+					ep[1] = es->end.gsamples[start].height;
+					eq[1] = es->end.gsamples[i-1].height;
+
+					JumpLink* link = addLink();
+
+					link->action = es->action;
+					link->flags = VALID;
+					link->nspine = es->trajectory.nspine;
+
+					const float startx = es->trajectory.spine[0];
+					const float endx = es->trajectory.spine[(es->trajectory.nspine-1)*2];
+					const float deltax = endx - startx;
+					
+					const float starty = es->trajectory.spine[1];
+					const float endy = es->trajectory.spine[(es->trajectory.nspine-1)*2+1];
+					
+					// Build link->spine0
+					for (int j = 0; j < es->trajectory.nspine; ++j)
+					{
+						const float* spt = &es->trajectory.spine[j*2];
+						const float u = (spt[0] - startx)/deltax;
+						const float dy = spt[1] - dtLerp(starty, endy, u) + m_linkBuilderConfig.agentClimb;
+						dtReal* p = &link->spine0[j*3];
+						dtVlerp(p, sp, ep, u);
+						dtVmad(p, p, es->ay, dy);
+					}
+
+					// Build link->spine1
+					for (int j = 0; j < es->trajectory.nspine; ++j)
+					{
+						const float* spt = &es->trajectory.spine[j*2];
+						const float u = (spt[0] - startx)/deltax;
+						const float dy = spt[1] - dtLerp(starty, endy, u) + m_linkBuilderConfig.agentClimb;
+						dtReal* p = &link->spine1[j*3];
+						dtVlerp(p, sq, eq, u);
+						dtVmad(p, p, es->ay, dy);
+					}
+				}
+				
+				start = -1;
+			}
+		}
+	}
+
+	dtFree(nflags, DT_ALLOC_TEMP);
+}
+
+void dtNavLinkBuilder::filterJumpOverLinks() const
+{
+	using namespace UE::Detour::NavLink::Private;
+	
+	// Filter out links which overlap
+	const float thresholdSquared = dtSqr(m_chf->cs*5.0f);
+	
+	for (int i = 0; i < m_nlinks-1; ++i)
+	{
+		JumpLink* li = &m_links[i];
+		if (li->flags == INVALID)
+			continue;
+		
+		const dtReal* spi = &li->spine0[0];
+		const dtReal* sqi = &li->spine1[0];
+		const dtReal* epi = &li->spine0[(li->nspine-1)*3];
+		const dtReal* eqi = &li->spine1[(li->nspine-1)*3];
+		
+		for (int j = i+1; j < m_nlinks; ++j)
+		{
+			JumpLink* lj = &m_links[j];
+			if (lj->flags == INVALID)
+				continue;
+			
+			const dtReal* spj = &lj->spine0[0];
+			const dtReal* sqj = &lj->spine1[0];
+			const dtReal* epj = &lj->spine0[(lj->nspine-1)*3];
+			const dtReal* eqj = &lj->spine1[(lj->nspine-1)*3];
+			
+			const dtReal d0 = distSegSegSqr(spi,sqi, epj,eqj);
+			const dtReal d1 = distSegSegSqr(epi,eqi, spj,sqj);
+			
+			if (d0 < thresholdSquared && d1 < thresholdSquared)
+			{
+				if (dtVdistSqr(spi,sqi) > dtVdistSqr(spj,sqj))
+				{
+					lj->flags = INVALID;
+				}
+				else
+				{
+					li->flags = INVALID;
+					break;
+				}
+			}
+		}
+	}
+}
+
+void dtNavLinkBuilder::buildForAllEdges(const dtLinkBuilderConfig& builderConfig, dtNavLinkAction action)
+{
+	for (int i = 0; i < m_nedges; ++i)
+	{
+		EdgeSampler sampler;
+		const bool success = sampleEdge(action, m_edges[i].sp, m_edges[i].sq, &sampler);
+		if (success)
+		{
+			addEdgeLinks(builderConfig, &sampler);
+		}
+	}
+	
+	filterJumpOverLinks();
+}
+
+void dtNavLinkBuilder::debugBuildEdge(const dtLinkBuilderConfig& builderConfig, dtNavLinkAction action, int edgeIndex, EdgeSampler& sampler)
+{
+	if (edgeIndex >= m_nedges)
+	{
+		return;
+	}
+	
+	m_debugSelectedEdge = edgeIndex;
+
+	const bool success = sampleEdge(action, m_edges[edgeIndex].sp, m_edges[edgeIndex].sq, &sampler);
+	if (success)
+	{
+		addEdgeLinks(builderConfig, &sampler);
+	}
+
+	filterJumpOverLinks();
+}
+
+bool dtNavLinkBuilder::getCompactHeightfieldHeight(const dtReal* pt, const float hrange, dtReal* height) const
+{
+	const float range = m_chf->cs;
+	const int ix0 = dtClamp((int)floorf((pt[0]-range - m_chf->bmin[0])/m_chf->cs), 0, m_chf->width-1);
+	const int iz0 = dtClamp((int)floorf((pt[2]-range - m_chf->bmin[2])/m_chf->cs), 0, m_chf->height-1);
+	const int ix1 = dtClamp((int)floorf((pt[0]+range - m_chf->bmin[0])/m_chf->cs), 0, m_chf->width-1);
+	const int iz1 = dtClamp((int)floorf((pt[2]+range - m_chf->bmin[2])/m_chf->cs), 0, m_chf->height-1);
+	
+	dtReal bestDist = DT_REAL_MAX;
+	dtReal bestHeight = DT_REAL_MAX;
+	bool found = false;
+	
+	for (int z = iz0; z <= iz1; ++z)
+	{
+		for (int x = ix0; x <= ix1; ++x)
+		{
+			const rcCompactCell& c = m_chf->cells[x+z*m_chf->width];
+			for (int i = (int)c.index, ni = (int)(c.index+c.count); i < ni; ++i)
+			{
+				const rcCompactSpan& s = m_chf->spans[i];
+				if (m_chf->areas[i] == RC_NULL_AREA)
+				{
+					continue;
+				}
+				
+				const dtReal y = m_chf->bmin[1] + s.y * m_chf->ch;
+				const dtReal dist = abs(y - pt[1]);
+				if (dist < hrange && dist < bestDist)
+				{
+					bestDist = dist;
+					bestHeight = y;
+					found = true;
+				}
+			}
+		}
+	}
+	
+	if (found)
+	{
+		*height = bestHeight;
+	}
+	else
+	{
+		*height = pt[1];
+	}
+	
+	return found;
+}
+
+// Compare ymin, ymax range with the solid height field spans to see if it collides.
+// Returns true if there is a collision.
+bool dtNavLinkBuilder::checkHeightfieldCollision(const dtReal x, const dtReal ymin, const dtReal ymax, const dtReal z) const
+{
+	using namespace UE::Detour::NavLink::Private;
+	
+	const int w = m_solid->width;
+	const int h = m_solid->height;
+	const dtReal cs = m_solid->cs;
+	const dtReal ch = m_solid->ch;
+	const rcReal* orig = m_solid->bmin;
+	const int ix = (int)floorf((x - orig[0])/cs);
+	const int iz = (int)floorf((z - orig[2])/cs);
+	
+	if (ix < 0 || iz < 0 || ix > w || iz > h)
+	{
+		return false;
+	}
+	
+	const rcSpan* s = m_solid->spans[ix + iz*w];
+	if (!s)
+	{
+		return false;
+	}
+	
+	while (s)
+	{
+		const float symin = orig[1] + s->data.smin*ch;
+		const float symax = orig[1] + s->data.smax*ch;
+		if (overlapRange(ymin, ymax, symin, symax))
+			return true;
+		
+		s = s->next;
+	}
+
+	return false;
+}
+
+// Returns true if none of the samples ymin, ymax collide with the heghtfield.
+bool dtNavLinkBuilder::isTrajectoryClear(const dtReal* pa, const dtReal* pb, const Trajectory2D* tra) const
+{
+	for (int i = 0; i < tra->nsamples; ++i)
+	{
+		dtReal p[3];
+		const TrajectorySample* s = &tra->samples[i];
+		const float u = (float)i / (float)(tra->nsamples-1);
+		dtVlerp(p, pa, pb, u);
+		if (checkHeightfieldCollision(p[0], p[1] + s->ymin, p[1] + s->ymax, p[2]))
+		{
+			return false;
+		}
+	}
+	
+	return true;	
+}
+
+void dtNavLinkBuilder::sampleGroundSegment(GroundSegment* seg, const float nsamples, const float groundRange) const
+{
+	dtReal delta[3];
+	dtVsub(delta, seg->p, seg->q);
+	
+	seg->ngsamples = nsamples;
+	// TODO: profile, preallocate, use dtChunkArray or else.
+	dtAssert(seg->gsamples == nullptr);
+	seg->gsamples = (GroundSample*)dtAlloc(sizeof(GroundSample)*seg->ngsamples, DT_ALLOC_PERM_TILE_LINK_BUILDER);	 
+	seg->npass = 0;
+	
+	for (int i = 0; i < seg->ngsamples; ++i)
+	{
+		const float u = (float)i/(float)(seg->ngsamples-1);
+		dtReal pt[3];
+
+		GroundSample* s = &seg->gsamples[i];
+		dtVlerp(pt, seg->p, seg->q, u);
+		s->flags = dtNavLinkBuilder::UNSET;
+		if (!getCompactHeightfieldHeight(pt, groundRange, &s->height))
+		{
+			continue;
+		}
+		
+		s->flags = static_cast<GroundSampleFlag>((unsigned char)s->flags | (unsigned char)HAS_GROUND);
+		seg->npass++;
+	}
+}
+
+void dtNavLinkBuilder::sampleAction(const EdgeSampler* es) const
+{
+	if (es->start.ngsamples != es->end.ngsamples)
+		return;
+	
+	const int nsamples = es->start.ngsamples;
+	
+	for (int i = 0; i < nsamples; ++i)
+	{
+		GroundSample* ssmp = &es->start.gsamples[i];
+		GroundSample* esmp = &es->end.gsamples[i];
+		
+		if ((ssmp->flags & HAS_GROUND) == 0 || (esmp->flags & HAS_GROUND) == 0)
+			continue;
+
+		const dtReal u = (dtReal)i/(dtReal)(nsamples-1);
+		dtReal spt[3], ept[3];
+		dtVlerp(spt, es->start.p, es->start.q, u);
+		dtVlerp(ept, es->end.p, es->end.q, u);
+		
+		spt[1] = ssmp->height;
+		ept[1] = esmp->height;
+		
+		if (!isTrajectoryClear(spt, ept, &es->trajectory))
+			continue;
+
+		ssmp->flags = static_cast<GroundSampleFlag>((unsigned char)ssmp->flags | (unsigned char)UNRESTRICTED);
+	}
+}
+
+void dtNavLinkBuilder::initTrajectory(Trajectory2D* tra) const
+{
+	using namespace UE::Detour::NavLink::Private;
+	
+	const float cs = m_linkBuilderConfig.cellSize;
+	
+	const float* pa = &tra->spine[0];
+	const float* pb = &tra->spine[(tra->nspine-1)*2];
+	
+	const float dx = pb[0] - pa[0];
+	tra->nsamples = dtMax(2, (int)ceilf(dx/cs));
+	// TODO: profile, preallocate, use dtChunkArray or else.
+	dtAssert(tra->samples == nullptr);
+	tra->samples = (TrajectorySample*)dtAlloc(sizeof(TrajectorySample)*tra->nsamples, DT_ALLOC_PERM_TILE_LINK_BUILDER);
+	
+	for (int i = 0; i < tra->nsamples; ++i)
+	{
+		const float u = (float)i / (float)(tra->nsamples-1);
+		TrajectorySample* s = &tra->samples[i];
+		s->x = dtLerp(pa[0], pb[0], u);
+		
+		const float y0 = getHeight(s->x-m_linkBuilderConfig.agentRadius, tra->spine, tra->nspine);
+		const float y1 = getHeight(s->x+m_linkBuilderConfig.agentRadius, tra->spine, tra->nspine);
+
+		const float y = dtLerp(pa[1], pb[1], u);
+		
+		s->ymin = dtMin(y0,y1) + m_linkBuilderConfig.agentClimb - y;
+		s->ymax = dtMax(y0,y1) + m_linkBuilderConfig.agentHeight - y; 
+	}
+}
+
+int dtNavLinkBuilder::findPotentialJumpOverEdges(const dtReal* sp, const dtReal* sq,
+												  const float depthRange, const float heightRange,
+												  dtReal* outSegs, const int maxOutSegs) const
+{
+	using namespace UE::Detour::NavLink::Private;
+	
+	// Find potential edges to join to.
+	const float widthRange = sqrtf(dtVdistSqr(sp,sq));
+	const float amin[3] = {0, -heightRange*0.5f, 0};
+	const float amax[3] = { widthRange, heightRange*0.5f, depthRange};
+	
+	const dtReal thr = cosf((180 - 45)/180*RC_PI);
+	
+	dtReal ax[3], ay[3], az[3];
+	dtVsub(ax, sq, sp);
+	dtVnormalize(ax);
+	dtVset(az, ax[2], 0, -ax[0]);
+	dtVnormalize(az);
+	dtVset(ay, 0, 1, 0);
+	
+	static constexpr int MAX_SEGS = 64;
+	PotentialSeg segs[MAX_SEGS];
+	int nsegs = 0;
+	
+	for (int i = 0; i < m_nedges; ++i)
+	{
+		dtReal p[3], lsp[3], lsq[3];
+		dtVsub(p, m_edges[i].sp, sp);
+		lsp[0] = ax[0]*p[0] + ay[0]*p[1] + az[0]*p[2];
+		lsp[1] = ax[1]*p[0] + ay[1]*p[1] + az[1]*p[2];
+		lsp[2] = ax[2]*p[0] + ay[2]*p[1] + az[2]*p[2];
+		
+		dtVsub(p, m_edges[i].sq, sp);
+		lsq[0] = ax[0]*p[0] + ay[0]*p[1] + az[0]*p[2];
+		lsq[1] = ax[1]*p[0] + ay[1]*p[1] + az[1]*p[2];
+		lsq[2] = ax[2]*p[0] + ay[2]*p[1] + az[2]*p[2];
+		
+		float tmin, tmax;
+		if (isectSegAABB(lsp, lsq, amin, amax, tmin, tmax))
+		{
+			if (tmin > 1.0f)
+				continue;
+			if (tmax < 0.0f)
+				continue;
+			
+			dtReal edir[3];
+			dtVsub(edir, m_edges[i].sq, m_edges[i].sp);
+			edir[1] = 0;
+			dtVnormalize(edir);
+			
+			if (dtVdot(ax, edir) > thr)
+				continue;
+			
+			if (nsegs < MAX_SEGS)
+			{
+				segs[nsegs].umin = dtClamp(tmin, 0.0f, 1.0f);
+				segs[nsegs].umax = dtClamp(tmax, 0.0f, 1.0f);
+				segs[nsegs].dmin = dtMin(lsp[2], lsq[2]);
+				segs[nsegs].dmax = dtMax(lsp[2], lsq[2]);
+				segs[nsegs].idx = i;
+				segs[nsegs].mark = 0;
+				nsegs++;
+			}
+		}
+	}
+	
+	const float eps = m_chf->cs;
+	unsigned char mark = 1;
+	for (int i = 0; i < nsegs; ++i)
+	{
+		if (segs[i].mark != 0)
+			continue;
+
+		segs[i].mark = mark;
+		
+		for (int j = i+1; j < nsegs; ++j)
+		{
+			if (overlapRange(segs[i].dmin-eps, segs[i].dmax+eps,
+						 	 segs[j].dmin-eps, segs[j].dmax+eps))
+			{
+				segs[j].mark = mark;
+			}
+		}
+		
+		mark++;
+	}
+
+
+	int nout = 0;
+	
+	for (int i = 1; i < mark; ++i)
+	{
+		// Find destination mid point.
+		constexpr float toUU = 100.f;
+		float umin = 10.0f * toUU;
+		float umax = -10.0 * toUU;
+		dtReal ptmin[3], ptmax[3];
+		
+		for (int j = 0; j < nsegs; ++j)
+		{
+			PotentialSeg* seg = &segs[j];
+			if (seg->mark != (unsigned char)i)
+				continue;
+			
+			dtReal pa[3], pb[3];
+			dtVlerp(pa, m_edges[seg->idx].sp, m_edges[seg->idx].sq, seg->umin);
+			dtVlerp(pb, m_edges[seg->idx].sp, m_edges[seg->idx].sq, seg->umax);
+			
+			const float ua = getClosestPtPtSeg(pa, sp, sq);
+			const float ub = getClosestPtPtSeg(pb, sp, sq);
+			
+			if (ua < umin)
+			{
+				dtVcopy(ptmin, pa);
+				umin = ua;
+			}
+			if (ua > umax)
+			{
+				dtVcopy(ptmax, pa);
+				umax = ua;
+			}
+			
+			if (ub < umin)
+			{
+				dtVcopy(ptmin, pb);
+				umin = ub;
+			}
+			if (ub > umax)
+			{
+				dtVcopy(ptmax, pb);
+				umax = ub;
+			}
+		}
+		
+		if (umin > umax)
+			continue;
+		
+		dtReal end[3];
+		dtVlerp(end, ptmin, ptmax, 0.5f);
+		
+		dtReal start[3];
+		dtVlerp(start, sp, sq, (umin+umax)*0.5f);
+		
+		dtReal orig[3];
+		dtVlerp(orig, start, end, 0.5f);
+		
+		dtReal dir[3], norm[3];
+		dtVsub(dir, end, start);
+		dir[1] = 0;
+		dtVnormalize(dir);
+		dtVset(norm, dir[2], 0, -dir[0]);
+		
+		dtReal ssp[3], ssq[3];
+		
+		const dtReal width = widthRange * (umax-umin);
+		dtVmad(ssp, orig, norm, width*0.5f);
+		dtVmad(ssq, orig, norm, -width*0.5f);
+		
+		if (nout < maxOutSegs)
+		{
+			dtVcopy(&outSegs[nout*6+0], ssp);
+			dtVcopy(&outSegs[nout*6+3], ssq);
+			nout++;
+		}
+	}
+		
+	return nout;
+}
+	
+
+void dtNavLinkBuilder::initJumpDownRig(EdgeSampler* es, const dtReal* sp, const dtReal* sq,
+										const float jumpStartDist, const float jumpEndDist,
+										const float jumpDownDist, const float groundRange)
+{
+	es->action = DT_LINK_ACTION_JUMP_DOWN;
+	dtVcopy(es->rigp, sp);
+	dtVcopy(es->rigq, sq);
+	
+	dtVsub(es->ax, sq, sp);
+	dtVnormalize(es->ax);
+	dtVset(es->az, es->ax[2], 0, -es->ax[0]);
+	dtVnormalize(es->az);
+	dtVset(es->ay, 0, 1, 0);
+	
+	// Build action sampling spine.
+	es->trajectory.nspine = MAX_SPINE;
+	for (int i = 0; i < MAX_SPINE; ++i)
+	{
+		float* pt = &es->trajectory.spine[i*2];
+		const float u = (float)i/(float)(MAX_SPINE-1);
+		pt[0] = jumpStartDist + u * (jumpEndDist - jumpStartDist);
+		pt[1] = u*u*u * jumpDownDist;
+	}
+
+	es->groundRange = groundRange;
+}
+
+void dtNavLinkBuilder::initJumpOverRig(EdgeSampler* es, const dtReal* sp, const dtReal* sq,
+										const float jumpStartDist, const float jumpEndDist,
+										const float jumpHeight, const float groundRange)
+{
+	es->action = DT_LINK_ACTION_JUMP_OVER;
+	dtVcopy(es->rigp, sp);
+	dtVcopy(es->rigq, sq);
+	
+	dtVsub(es->ax, sq, sp);
+	dtVnormalize(es->ax);
+	dtVset(es->az, es->ax[2], 0, -es->ax[0]);
+	dtVnormalize(es->az);
+	dtVset(es->ay, 0, 1, 0);
+	
+	// Build action sampling spine.
+	es->trajectory.nspine = MAX_SPINE;
+	for (int i = 0; i < MAX_SPINE; ++i)
+	{
+		float* pt = &es->trajectory.spine[i*2];
+		const float u = (float)i/(float)(MAX_SPINE-1);
+		pt[0] = jumpStartDist + u * (jumpEndDist - jumpStartDist);
+		pt[1] = (1-dtSqr(u*2-1)) * jumpHeight;
+	}
+	
+	es->groundRange = groundRange;
+}
+
+bool dtNavLinkBuilder::sampleEdge(dtNavLinkAction desiredAction, const dtReal* sp, const dtReal* sq, dtNavLinkBuilder::EdgeSampler* es) const
+{
+	using namespace UE::Detour::NavLink::Private;
+
+	constexpr float toUU = 100.f;
+	
+	if (desiredAction == DT_LINK_ACTION_JUMP_DOWN)
+	{
+		constexpr float jumpStartDist = -0.25f * toUU;
+		constexpr float jumpEndDist = 2.f * toUU;
+		constexpr float jumpDownDist = -2.f * toUU;
+		constexpr float groundRange = 1.f * toUU;
+		initJumpDownRig(es, sp, sq, jumpStartDist, jumpEndDist, jumpDownDist, groundRange);
+	}
+	else if (desiredAction == DT_LINK_ACTION_JUMP_OVER)
+	{
+		constexpr float jumpDist = 2.5f * toUU;
+		constexpr float heightRange = 1.f * toUU;
+		static constexpr int NSEGS = 8;
+		dtReal segs[NSEGS*6];
+		int nsegs = findPotentialJumpOverEdges(sp, sq, jumpDist, heightRange, segs, NSEGS);
+		int ibest = -1;
+		float dbest = 0;
+		for (int i = 0; i < nsegs; ++i)
+		{
+			const dtReal* seg = &segs[i*6];
+			const float d = dtVdistSqr(seg,seg+3);
+			if (d > dbest)
+			{
+				dbest = d;
+				ibest = i;
+			}
+		}
+		if (ibest == -1)
+		{
+			return false;
+		}
+
+		constexpr float jumpHeight = 1.f * toUU;
+		constexpr float groundRange = 0.5f * toUU;
+		initJumpOverRig(es, &segs[ibest*6+0], &segs[ibest*6+3], -jumpDist*0.5f, jumpDist*0.5f, jumpHeight, groundRange);
+	}
+	
+	initTrajectory(&es->trajectory);
+
+	// Init start end segments.
+	dtReal offset[3];
+	trans2d(offset, es->az, es->ay, &es->trajectory.spine[0]);
+	dtVadd(es->start.p, es->rigp, offset);
+	dtVadd(es->start.q, es->rigq, offset);
+	trans2d(offset, es->az, es->ay, &es->trajectory.spine[(es->trajectory.nspine-1)*2]);
+	dtVadd(es->end.p, es->rigp, offset);
+	dtVadd(es->end.q, es->rigq, offset);
+
+	// Sample start and end ground segments.
+	const float dist = sqrtf(dtVdistSqr(es->rigp, es->rigq));
+	const int ngsamples = dtMax(2, (int)ceilf(dist/m_cs)); // One trajectory per cs
+	
+	sampleGroundSegment(&es->start, ngsamples, es->groundRange);
+	sampleGroundSegment(&es->end, ngsamples, es->groundRange);
+
+	sampleAction(es);
+	
+	return true;
+}
+//@UE END
diff --git a/Engine/Source/Runtime/Navmesh/Private/DetourTileCache/DetourTileCacheBuilder.cpp b/Engine/Source/Runtime/Navmesh/Private/DetourTileCache/DetourTileCacheBuilder.cpp
--- a/Engine/Source/Runtime/Navmesh/Private/DetourTileCache/DetourTileCacheBuilder.cpp
+++ b/Engine/Source/Runtime/Navmesh/Private/DetourTileCache/DetourTileCacheBuilder.cpp
@@ -152,10 +152,10 @@ struct dtTempContour
 		poly(pbuf), npoly(0), cpoly(npbuf) 
 	{
 	}
-	unsigned short* verts;	// x, y, z, reg, area
+	unsigned short* verts; // v[0], v[1], v[2] coordinates, v[3] region, v[4] area (high bit indicates vertex pinning, removed in countour simplification).
 	int nverts;
 	int cverts;
-	unsigned short* poly;
+	unsigned short* poly; // index in verts for the contour of the poly
 	int npoly;
 	int cpoly;
 };
@@ -348,6 +348,7 @@ static bool walkContour(dtTileCacheLayer& layer, int x, int y, int idx, const bo
 
 		getNeighbourRegAndAreaAndVertexHeight(layer, x, y, dir, neiReg, neiArea, cornerNeiArea, vertexHeight);	// UE
 
+		// Check if the region in the provided direction is different than the region at x,y.
 		if (neiReg != layer.regs[x+y*w])
 		{
 			// Solid edge.
@@ -954,6 +955,7 @@ dtStatus dtBuildTileCacheContours(dtTileCacheAlloc* alloc, dtTileCacheLayer& lay
 			const int idx = x+y*w;
 			if (flags[idx] == 0)
 			{
+				// All cell edges are connected, ignore it.
 				continue;
 			}
 
@@ -1034,7 +1036,7 @@ dtStatus dtBuildTileCacheContours(dtTileCacheAlloc* alloc, dtTileCacheLayer& lay
 					dst[2] = v[2];
 
 					// Store portal direction and store remove status to the fourth component.
-					dst[3] = 0x0f;
+					dst[3] = 0x0f;	// Set direction to 0xf
 					if (nei != 0xffff && nei >= 0xf800)
 						dst[3] = (unsigned char)(nei - 0xf800);
 					if (shouldRemove)
@@ -1233,7 +1235,7 @@ static unsigned short addVertex(unsigned short x, unsigned short y, unsigned sho
 
 namespace TileCacheData
 {
-	struct rcEdge
+	struct dtEdge
 	{
 		unsigned short vert[2];			// index in verts (a,b)
 		unsigned short polyEdge[2];		// index in polys (a,b)
@@ -1257,7 +1259,7 @@ static bool buildMeshAdjacency(dtTileCacheAlloc* alloc,
 	unsigned short* nextEdge = firstEdge + nverts;
 	int edgeCount = 0;
 	
-	dtFixedArray<TileCacheData::rcEdge> edges(alloc, maxEdgeCount);
+	dtFixedArray<TileCacheData::dtEdge> edges(alloc, maxEdgeCount);
 	if (!edges)
 		return false;
 	
@@ -1277,7 +1279,7 @@ static bool buildMeshAdjacency(dtTileCacheAlloc* alloc,
 			const unsigned short v1 = (j+1 >= MAX_VERTS_PER_POLY || t[j+1] == DT_TILECACHE_NULL_IDX) ? t[0] : t[j+1];
 			if (v0 < v1)
 			{
-				TileCacheData::rcEdge& edge = edges[edgeCount];
+				TileCacheData::dtEdge& edge = edges[edgeCount];
 				edge.vert[0] = v0;
 				edge.vert[1] = v1;
 				edge.poly[0] = (unsigned short)i;
@@ -1308,7 +1310,7 @@ static bool buildMeshAdjacency(dtTileCacheAlloc* alloc,
 				bool found = false;
 				for (unsigned short e = firstEdge[v1]; e != DT_TILECACHE_NULL_IDX; e = nextEdge[e])
 				{
-					TileCacheData::rcEdge& edge = edges[e];
+					TileCacheData::dtEdge& edge = edges[e];
 					if (edge.vert[1] == v0 && edge.poly[0] == edge.poly[1])
 					{
 						// Edges matches
@@ -1322,7 +1324,7 @@ static bool buildMeshAdjacency(dtTileCacheAlloc* alloc,
 				if (!found)
 				{
 					// Matching edge not found, it is an open edge, add it.
-					TileCacheData::rcEdge& edge = edges[edgeCount];
+					TileCacheData::dtEdge& edge = edges[edgeCount];
 					edge.vert[0] = v1;
 					edge.vert[1] = v0;
 					edge.poly[0] = (unsigned short)i;
@@ -1371,7 +1373,7 @@ static bool buildMeshAdjacency(dtTileCacheAlloc* alloc,
 				
 				for (int m = 0; m < edgeCount; ++m)
 				{
-					TileCacheData::rcEdge& e = edges[m];
+					TileCacheData::dtEdge& e = edges[m];
 					// Skip connected edges.
 					if (e.poly[0] != e.poly[1])
 						continue;
@@ -1414,11 +1416,10 @@ static bool buildMeshAdjacency(dtTileCacheAlloc* alloc,
 				
 				for (int m = 0; m < edgeCount; ++m)
 				{
-					TileCacheData::rcEdge& e = edges[m];
+					TileCacheData::dtEdge& e = edges[m];
 					// Skip connected edges.
 					if (e.poly[0] != e.poly[1])
 						continue;
-
 					const unsigned short* eva = &verts[e.vert[0]*3];
 					const unsigned short* evb = &verts[e.vert[1]*3];
 					if (eva[2] == z && evb[2] == z)
@@ -1446,11 +1447,13 @@ static bool buildMeshAdjacency(dtTileCacheAlloc* alloc,
 	
 	
 	// Store adjacency
+	// Adjacency between poly is store in the second arrays of polys.
 	for (int i = 0; i < edgeCount; ++i)
 	{
-		const TileCacheData::rcEdge& e = edges[i];
+		const TileCacheData::dtEdge& e = edges[i];
 		if (e.poly[0] != e.poly[1])
 		{
+			// If not the same poly, store the direction of the portal.
 			unsigned short* p0 = &polys[e.poly[0]*MAX_VERTS_PER_POLY*2];
 			unsigned short* p1 = &polys[e.poly[1]*MAX_VERTS_PER_POLY*2];
 			p0[MAX_VERTS_PER_POLY + e.polyEdge[0]] = e.poly[1];
@@ -1458,6 +1461,7 @@ static bool buildMeshAdjacency(dtTileCacheAlloc* alloc,
 		}
 		else if (e.polyEdge[1] != 0xff) // if we have a direction
 		{
+			// Same poly
 			unsigned short* p0 = &polys[e.poly[0]*MAX_VERTS_PER_POLY*2];
 			p0[MAX_VERTS_PER_POLY + e.polyEdge[0]] = 0x8000 | (unsigned short)e.polyEdge[1];
 		}
@@ -2293,7 +2297,7 @@ dtStatus dtBuildTileCachePolyMesh(dtTileCacheAlloc* alloc,
 	
 	for (int i = 0; i < lcset.nconts; ++i)
 	{
-		dtTileCacheContour& cont = lcset.conts[i];
+		const dtTileCacheContour& cont = lcset.conts[i];
 		
 		// Skip null contours.
 		if (cont.nverts < 3 || lcset.conts[i].area == DT_TILECACHE_NULL_AREA)
diff --git a/Engine/Source/Runtime/Navmesh/Public/DebugUtils/DebugDraw.h b/Engine/Source/Runtime/Navmesh/Public/DebugUtils/DebugDraw.h
--- a/Engine/Source/Runtime/Navmesh/Public/DebugUtils/DebugDraw.h
+++ b/Engine/Source/Runtime/Navmesh/Public/DebugUtils/DebugDraw.h
@@ -149,6 +149,16 @@ struct duDebugDraw
 	virtual void end() = 0;
 };
 
+// Reorder to handle the 0xRRGGBBAA format
+inline unsigned int duRGBA(unsigned int rgba)
+{
+	const unsigned int a = (rgba & 0x000000ff);
+	const unsigned int b = (rgba & 0x0000ff00);
+	const unsigned int g = (rgba & 0x00ff0000);
+	const unsigned int r = (rgba & 0xff000000);
+	return (r >> 24) | (g >> 8) | (b << 8) | a << 24;
+}
+
 inline unsigned int duRGBA(int r, int g, int b, int a)
 {
 	return ((unsigned int)r) | ((unsigned int)g << 8) | ((unsigned int)b << 16) | ((unsigned int)a << 24);
@@ -203,6 +213,35 @@ inline unsigned int duTransCol(unsigned int c, unsigned int a)
 	return (a<<24) | (c & 0x00ffffff);
 }
 
+//@UE BEGIN
+struct duColor
+{
+	// Color palette
+	static unsigned int red;
+	static unsigned int pink;
+	static unsigned int purple;
+	static unsigned int deepPurple;
+	static unsigned int indigo;
+	static unsigned int blue;
+	static unsigned int lightBlue;
+	static unsigned int cyan;
+	static unsigned int teal;		
+	static unsigned int green;
+	static unsigned int lightGreen;
+	static unsigned int lime;
+	static unsigned int yellow;
+	static unsigned int amber;
+	static unsigned int orange;
+	static unsigned int orangeRed;
+	static unsigned int brown;
+	
+	static unsigned int white;
+	static unsigned int lightGrey;
+	static unsigned int grey;
+	static unsigned int drakGrey;
+	static unsigned int black;
+};
+//@UE END
 
 void duCalcBoxColors(unsigned int* colors, unsigned int colTop, unsigned int colSide);
 
diff --git a/Engine/Source/Runtime/Navmesh/Public/DebugUtils/DetourNavLinkDebugDraw.h b/Engine/Source/Runtime/Navmesh/Public/DebugUtils/DetourNavLinkDebugDraw.h
new file mode 100644
--- /dev/null
+++ b/Engine/Source/Runtime/Navmesh/Public/DebugUtils/DetourNavLinkDebugDraw.h
@@ -0,0 +1,23 @@
+// Copyright Epic Games, Inc. All Rights Reserved.
+#pragma once
+
+#include "CoreTypes.h"
+#include "Detour/DetourNavLinkBuilder.h"
+
+//@UE BEGIN
+struct duDebugDraw;
+class dtNavLinkBuilder;
+
+enum duNavLinkBuilderDrawFlags
+{
+	DRAW_WALKABLE_SURFACE =	1 << 0,
+	DRAW_BORDERS =			1 << 1,
+	DRAW_SELECTED_EDGE =	1 << 2,
+	DRAW_TRAJECTORY =		1 << 3,
+	DRAW_LAND_SAMPLES =		1 << 4,
+	DRAW_COLLISION_SLICES =	1 << 5,
+	DRAW_LINKS =			1 << 6,
+};
+
+NAVMESH_API void duDebugDrawNavLinkBuilder(duDebugDraw* dd, const dtNavLinkBuilder& linkBuilder, unsigned int drawFlags, const struct dtNavLinkBuilder::EdgeSampler* es);
+//@UE END
diff --git a/Engine/Source/Runtime/Navmesh/Public/Detour/DetourAlloc.h b/Engine/Source/Runtime/Navmesh/Public/Detour/DetourAlloc.h
--- a/Engine/Source/Runtime/Navmesh/Public/Detour/DetourAlloc.h
+++ b/Engine/Source/Runtime/Navmesh/Public/Detour/DetourAlloc.h
@@ -44,6 +44,7 @@ enum dtAllocHint
 	DT_ALLOC_PERM_TILE_DYNLINK_OFFMESH,
 	DT_ALLOC_PERM_TILE_DYNLINK_CLUSTER,
 	DT_ALLOC_PERM_TILES,
+	DT_ALLOC_PERM_TILE_LINK_BUILDER,
 };
 //@UE END Adding support for memory tracking.
 
diff --git a/Engine/Source/Runtime/Navmesh/Public/Detour/DetourCommon.h b/Engine/Source/Runtime/Navmesh/Public/Detour/DetourCommon.h
--- a/Engine/Source/Runtime/Navmesh/Public/Detour/DetourCommon.h
+++ b/Engine/Source/Runtime/Navmesh/Public/Detour/DetourCommon.h
@@ -155,6 +155,11 @@ inline double dtfMod(double x, double y)
 	return fmod(x, y);
 }
 
+inline float dtLerp(float a, float b, float t)
+{
+	return a + (b-a)*t;
+}
+
 /// @}
 /// @name Vector helper functions.
 /// @{
@@ -350,6 +355,14 @@ inline void dtVnormalize(dtReal* v)
 	v[2] *= d;
 }
 
+/// Returns true if the vectors are the same.
+///  @param[in]		a	A vector. [(x, y, z)]
+///  @param[in]		b	A vector. [(x, y, z)]
+inline bool dtVisEqual(const unsigned short* a, const unsigned short* b)
+{
+	return a[0] == b[0] && a[1] == b[1] && a[2] == b[2];
+}
+
 /// Performs a 'sloppy' colocation check of the specified points.
 ///  @param[in]		p0	A point. [(x, y, z)]
 ///  @param[in]		p1	A point. [(x, y, z)]
diff --git a/Engine/Source/Runtime/Navmesh/Public/Detour/DetourNavLinkBuilder.h b/Engine/Source/Runtime/Navmesh/Public/Detour/DetourNavLinkBuilder.h
new file mode 100644
--- /dev/null
+++ b/Engine/Source/Runtime/Navmesh/Public/Detour/DetourNavLinkBuilder.h
@@ -0,0 +1,189 @@
+// Copyright Epic Games, Inc. All Rights Reserved.
+#pragma once
+
+#include "CoreTypes.h"
+#include "DetourAlloc.h"
+#include "Detour/DetourLargeWorldCoordinates.h"
+
+//@UE BEGIN
+enum dtNavLinkAction
+{
+	DT_LINK_ACTION_UNSET,
+	DT_LINK_ACTION_JUMP_DOWN,
+	DT_LINK_ACTION_JUMP_OVER,
+};
+
+struct dtLinkBuilderConfig
+{
+	dtReal agentRadius = 0;
+	dtReal agentHeight = 0;
+	dtReal agentClimb = 0;
+	dtReal cellSize = 0;
+	dtReal cellHeight = 0;
+};
+
+struct rcHeightfield;
+struct rcCompactHeightfield;
+
+NAVMESH_API struct dtLinkBuilderData
+{
+	bool generatingLinks = false;
+	rcHeightfield* solidHF = nullptr;
+	rcCompactHeightfield* compactHF = nullptr;
+};
+
+class rcContext;
+struct rcConfig;
+
+class dtNavLinkBuilder
+{
+	dtLinkBuilderConfig m_linkBuilderConfig;
+
+	dtReal m_cs = 0;
+	rcHeightfield* m_solid;
+	rcCompactHeightfield* m_chf;
+
+	struct Edge
+	{
+		dtReal sp[3], sq[3];
+	};
+	Edge* m_edges;
+	int m_nedges;
+	
+public:	
+	static constexpr int MAX_SPINE = 8;
+	
+	struct TrajectorySample
+	{
+		float x, ymin, ymax;
+	};
+	
+	struct Trajectory2D
+	{
+		Trajectory2D() : samples(nullptr), nsamples(0), nspine(0) {}
+		NAVMESH_API ~Trajectory2D();
+
+		float spine[2*MAX_SPINE];		// [x,y] relative points representing the desired trajectory (2 spines)
+		TrajectorySample* samples;		// samples along trajectory slices to check for collision
+		
+		unsigned short nsamples;
+		unsigned char nspine;			// @todo: remove (use direclty MAX_SPINE) or make relative to trajectory type and config
+	};
+
+private:	
+	enum GroundSampleFlag : unsigned char
+	{
+		UNSET = 0,
+		HAS_GROUND = 1,
+		UNRESTRICTED = 4,
+	};
+	
+	struct GroundSample
+	{
+		dtReal height;
+		GroundSampleFlag flags;
+	};
+
+	struct PotentialSeg
+	{
+		unsigned char mark;
+		int idx;
+		float umin, umax;
+		float dmin, dmax;
+		float sp[3], sq[3];
+	};
+	
+	struct GroundSegment
+	{
+		GroundSegment() : gsamples(nullptr), ngsamples(0) {}
+		NAVMESH_API ~GroundSegment();
+
+		dtReal p[3], q[3];
+		GroundSample* gsamples;
+		unsigned short ngsamples;
+		unsigned short npass;
+	};
+	
+public:
+	NAVMESH_API struct EdgeSampler
+	{
+		Trajectory2D trajectory;
+	
+		GroundSegment start;
+		GroundSegment end;
+
+		dtReal rigp[3], rigq[3];		// edge
+		dtReal ax[3], ay[3], az[3];		// axis along edge
+
+		float groundRange;
+		
+		dtNavLinkAction action = DT_LINK_ACTION_UNSET;
+	};
+
+	enum JumpLinkFlag : unsigned char
+	{
+		INVALID = 0,
+		VALID = 1,
+	};
+	
+	struct JumpLink
+	{
+		dtReal spine0[MAX_SPINE*3];
+		dtReal spine1[MAX_SPINE*3];
+		int nspine;
+		JumpLinkFlag flags;
+		dtNavLinkAction action = DT_LINK_ACTION_UNSET;
+	};
+	JumpLink* m_links;
+	int m_nlinks;
+	int m_clinks;
+
+private:
+	int m_debugSelectedEdge;
+
+	friend NAVMESH_API void duDebugDrawNavLinkBuilder(struct duDebugDraw* dd, const dtNavLinkBuilder& linkBuilder, unsigned int drawFlags, const EdgeSampler* es);
+	
+public:
+	NAVMESH_API dtNavLinkBuilder();
+	NAVMESH_API ~dtNavLinkBuilder();
+
+	// Loops through contours to store edge points in world coordinates.
+	NAVMESH_API bool findEdges(rcContext& ctx, const rcConfig& cfg, const dtLinkBuilderConfig& builderConfig,
+							   const struct dtTileCacheContourSet& lcset, const dtReal* orig, const dtLinkBuilderData& linkBuilderData);
+
+	// For all edges, sample edges (sampleEdge) and add links to m_links
+	NAVMESH_API void buildForAllEdges(const dtLinkBuilderConfig& acfg, dtNavLinkAction action);
+
+	NAVMESH_API void debugBuildEdge(const dtLinkBuilderConfig& acfg, dtNavLinkAction action, int edgeIndex, EdgeSampler& sampler);
+	
+private:
+	void cleanup();
+	void initTrajectory(Trajectory2D* tra) const;
+	bool isTrajectoryClear(const dtReal* pa, const dtReal* pb, const Trajectory2D* tra) const;
+	
+	int findPotentialJumpOverEdges(const dtReal* sp, const dtReal* sq,
+								   const float depthRange, const float heightRange,
+								   dtReal* outSegs, const int maxOutSegs) const;
+	
+	static void initJumpDownRig(EdgeSampler* es, const dtReal* sp, const dtReal* sq,
+						 const float jumpStartDist, const float jumpEndDist,
+						 const float jumpDownDist, const float groundRange);
+
+	static void initJumpOverRig(EdgeSampler* es, const dtReal* sp, const dtReal* sq,
+						 const float jumpStartDist, const float jumpEndDist,
+						 const float jumpHeight, const float groundRange);
+
+	bool getCompactHeightfieldHeight(const dtReal* pt, const float hrange, dtReal* height) const;
+	bool checkHeightfieldCollision(const dtReal x, const dtReal ymin, const dtReal ymax, const dtReal z) const;
+
+	void sampleGroundSegment(GroundSegment* seg, const float nsamples, const float groundRange) const;
+	
+	void sampleAction(const EdgeSampler* es) const;
+	
+	void filterJumpOverLinks() const;
+
+	bool sampleEdge(dtNavLinkAction desiredAction, const dtReal* sp, const dtReal* sq, dtNavLinkBuilder::EdgeSampler* sampler) const;
+	JumpLink* addLink();
+	void addEdgeLinks(const dtLinkBuilderConfig& acfg, const EdgeSampler* es);
+};
+//@UE END
\ No newline at end of file
