Build fix in DetourNavLinkBuilder.h
#rnx

diff --git a/Engine/Source/Runtime/Navmesh/Public/Detour/DetourNavLinkBuilder.h b/Engine/Source/Runtime/Navmesh/Public/Detour/DetourNavLinkBuilder.h
--- a/Engine/Source/Runtime/Navmesh/Public/Detour/DetourNavLinkBuilder.h
+++ b/Engine/Source/Runtime/Navmesh/Public/Detour/DetourNavLinkBuilder.h
@@ -25,7 +25,7 @@ struct dtLinkBuilderConfig
 struct rcHeightfield;
 struct rcCompactHeightfield;
 
-NAVMESH_API struct dtLinkBuilderData
+struct NAVMESH_API dtLinkBuilderData
 {
 	bool generatingLinks = false;
 	rcHeightfield* solidHF = nullptr;
@@ -105,7 +105,7 @@ private:
 	};
 	
 public:
-	NAVMESH_API struct EdgeSampler
+	struct NAVMESH_API EdgeSampler
 	{
 		Trajectory2D trajectory;
 	
