Fix CIS issue in dtNavLinkBuilder
#rnx

diff --git a/Engine/Source/Runtime/Navmesh/Private/Detour/DetourNavLinkBuilder.cpp b/Engine/Source/Runtime/Navmesh/Private/Detour/DetourNavLinkBuilder.cpp
--- a/Engine/Source/Runtime/Navmesh/Private/Detour/DetourNavLinkBuilder.cpp
+++ b/Engine/Source/Runtime/Navmesh/Private/Detour/DetourNavLinkBuilder.cpp
@@ -753,7 +753,7 @@ int dtNavLinkBuilder::findPotentialJumpOverEdges(const dtReal* sp, const dtReal*
 	const float amin[3] = {0, -heightRange*0.5f, 0};
 	const float amax[3] = { widthRange, heightRange*0.5f, depthRange};
 	
-	const dtReal thr = cosf((180 - 45)/180*RC_PI);
+	const dtReal thr = cosf((180.0 - 45.0)/180.0*RC_PI);
 	
 	dtReal ax[3], ay[3], az[3];
 	dtVsub(ax, sq, sp);
@@ -837,7 +837,7 @@ int dtNavLinkBuilder::findPotentialJumpOverEdges(const dtReal* sp, const dtReal*
 		// Find destination mid point.
 		constexpr float toUU = 100.f;
 		float umin = 10.0f * toUU;
-		float umax = -10.0 * toUU;
+		float umax = -10.0f * toUU;
 		dtReal ptmin[3], ptmax[3];
 		
 		for (int j = 0; j < nsegs; ++j)
