Fix multiple CIS issues in DetourNavMeshQuery.cpp

diff --git a/Engine/Source/Runtime/Navmesh/Private/Detour/DetourNavMeshQuery.cpp b/Engine/Source/Runtime/Navmesh/Private/Detour/DetourNavMeshQuery.cpp
--- a/Engine/Source/Runtime/Navmesh/Private/Detour/DetourNavMeshQuery.cpp
+++ b/Engine/Source/Runtime/Navmesh/Private/Detour/DetourNavMeshQuery.cpp
@@ -779,17 +779,24 @@ void dtNavMeshQuery::closestPointOnPolyInTile(const dtMeshTile* tile, const dtPo
 	dtReal verts[DT_VERTS_PER_POLYGON*3];
 	dtReal edged[DT_VERTS_PER_POLYGON];
 	dtReal edget[DT_VERTS_PER_POLYGON];
-	const int nv = poly->vertCount;
+	const int nv = dtMin(poly->vertCount, DT_VERTS_PER_POLYGON);
 	for (int i = 0; i < nv; ++i)
+	{
 		dtVcopy(&verts[i*3], &tile->verts[poly->verts[i]*3]);
+	}
 	
 	dtVcopy(closest, pos);
+	if (nv == 0)
+	{
+		return;
+	}
+
 	if (!dtDistancePtPolyEdgesSqr(pos, verts, nv, edged, edget))
 	{
 		// Point is outside the polygon, dtClamp to nearest edge.
-		dtReal dmin = DT_REAL_MAX;
-		int imin = -1;
-		for (int i = 0; i < nv; ++i)
+		dtReal dmin = edged[0];
+		int imin = 0;
+		for (int i = 1; i < nv; ++i)
 		{
 			if (edged[i] < dmin)
 			{
@@ -797,9 +804,9 @@ void dtNavMeshQuery::closestPointOnPolyInTile(const dtMeshTile* tile, const dtPo
 				imin = i;
 			}
 		}
+		CA_ASSUME(imin < nv);
 		const dtReal* va = &verts[imin*3];
 		const dtReal* vb = &verts[((imin+1)%nv)*3];
-		CA_SUPPRESS(6385);
 		dtVlerp(closest, va, vb, edget[imin]);
 	}
 
@@ -864,7 +871,14 @@ dtStatus dtNavMeshQuery::closestPointOnPolyBoundary(dtPolyRef ref, const dtReal*
 	const dtMeshTile* tile = 0;
 	const dtPoly* poly = 0;
 	if (dtStatusFailed(m_nav->getTileAndPolyByRef(ref, &tile, &poly)))
+	{
 		return DT_FAILURE | DT_INVALID_PARAM;
+	}
+
+	if (poly->vertCount == 0)
+	{
+		return DT_FAILURE;
+	}
 	
 	// Collect vertices.
 	dtReal verts[DT_VERTS_PER_POLYGON*3];
@@ -886,9 +900,9 @@ dtStatus dtNavMeshQuery::closestPointOnPolyBoundary(dtPolyRef ref, const dtReal*
 	else
 	{
 		// Point is outside the polygon, dtClamp to nearest edge.
-		dtReal dmin = DT_REAL_MAX;
-		int imin = -1;
-		for (int i = 0; i < nv; ++i)
+		dtReal dmin = edged[0];
+		int imin = 0;
+		for (int i = 1; i < nv; ++i)
 		{
 			if (edged[i] < dmin)
 			{
@@ -896,9 +910,9 @@ dtStatus dtNavMeshQuery::closestPointOnPolyBoundary(dtPolyRef ref, const dtReal*
 				imin = i;
 			}
 		}
+		CA_ASSUME(imin < nv);
 		const dtReal* va = &verts[imin*3];
 		const dtReal* vb = &verts[((imin+1)%nv)*3];
-		CA_SUPPRESS(6385);
 		dtVlerp(closest, va, vb, edget[imin]);
 	}
 	
@@ -1153,7 +1167,12 @@ dtStatus dtNavMeshQuery::findNearestPoly(const dtReal* center, const dtReal* ext
 	{
 		dtPolyRef ref = polys[i];
 		dtReal closestPtPoly[3];
-		closestPointOnPoly(ref, referenceLocation, closestPtPoly);
+		const dtStatus result = closestPointOnPoly(ref, referenceLocation, closestPtPoly);
+		if (dtStatusFailed(result))
+		{
+			continue;
+		}
+
 		const dtReal d = dtVdistSqr(referenceLocation, closestPtPoly);
 		const dtReal h = dtAbs(center[1] - closestPtPoly[1]);
 //@UE END
@@ -1215,7 +1234,12 @@ dtStatus dtNavMeshQuery::findNearestPoly2D(const dtReal* center, const dtReal* e
 	{
 		dtPolyRef ref = polys[i];
 		dtReal closestPtPoly[3];
-		closestPointOnPoly(ref, referenceLocation, closestPtPoly);
+		const dtStatus result = closestPointOnPoly(ref, referenceLocation, closestPtPoly);
+		if (result & DT_FAILURE)
+		{
+			continue;
+		}
+
 		const dtReal dSq = dtVdist2DSqr(referenceLocation, closestPtPoly);
 		const dtReal h = dtAbs(center[1] - closestPtPoly[1]);
 
@@ -1311,7 +1335,12 @@ dtStatus dtNavMeshQuery::findNearestContainingPoly(const dtReal* center, const d
 		if (inPoly)
 		{
 			dtReal closestPtPoly[3];
-			closestPointOnPoly(ref, center, closestPtPoly);
+			const dtStatus result = closestPointOnPoly(ref, center, closestPtPoly);
+			if (result & DT_FAILURE)
+			{
+				continue;
+			}
+
 			const dtReal d = dtVdistSqr(center, closestPtPoly);
 			const dtReal x = dtAbs(center[0] - closestPtPoly[0]);
 			const dtReal z = dtAbs(center[2] - closestPtPoly[2]);
@@ -3260,7 +3289,7 @@ dtStatus dtNavMeshQuery::raycast(dtPolyRef startRef, const dtReal* startPos, con
 		m_nav->getTileAndPolyByRefUnsafe(curRef, &tile, &poly);
 		
 		// Check if poly has valid data, bail out otherwise
-		if (poly == nullptr || poly->vertCount > DT_VERTS_PER_POLYGON)
+		if (poly == nullptr || poly->vertCount > DT_VERTS_PER_POLYGON || poly->vertCount == 0)
 		{
 			if (pathCount)
 				*pathCount = n;
@@ -5045,6 +5074,10 @@ dtStatus dtNavMeshQuery::findDistanceToWall(dtPolyRef startRef, const dtReal* ce
 		const dtMeshTile* bestTile = 0;
 		const dtPoly* bestPoly = 0;
 		m_nav->getTileAndPolyByRefUnsafe(bestRef, &bestTile, &bestPoly);
+		if (bestPoly->vertCount == 0)
+		{
+			continue;
+		}
 		
 		// Get parent poly and tile.
 		dtPolyRef parentRef = 0;
