Fix CIS issues in DetourNavMesh

diff --git a/Engine/Source/Runtime/Navmesh/Private/Detour/DetourNavMesh.cpp b/Engine/Source/Runtime/Navmesh/Private/Detour/DetourNavMesh.cpp
--- a/Engine/Source/Runtime/Navmesh/Private/Detour/DetourNavMesh.cpp
+++ b/Engine/Source/Runtime/Navmesh/Private/Detour/DetourNavMesh.cpp
@@ -912,12 +912,22 @@ int dtNavMesh::findConnectingPolys(const dtReal* va, const dtReal* vb,
 	if (!tile) return 0;
 
 	dtReal amin[2], amax[2], apt[3];
+	amin[0] = 0;
+	amin[1] = 0;
+	amax[0] = 0;
+	amax[1] = 0;
+	
 	calcSlabEndPoints(va, vb, amin, amax, side);
 	const dtReal apos = getSlabCoord(va, side);
 	dtVcopy(apt, va);
 
 	// Remove links pointing to 'side' and compact the links array. 
 	dtReal bmin[2], bmax[2], bpt[3];
+	bmin[0] = 0;
+	bmin[1] = 0;
+	bmax[0] = 0;
+	bmax[1] = 0;
+	
 	unsigned short m = DT_EXT_LINK | (unsigned short)side;
 	int n = 0;
 
@@ -1443,12 +1453,17 @@ void dtNavMesh::closestPointOnPolyInTile(const dtMeshTile* tile, unsigned int ip
 		dtVcopy(&verts[i*3], &tile->verts[poly->verts[i]*3]);
 	
 	dtVcopy(closest, pos);
+	if (nv == 0)
+	{
+		return;
+	}
+	
 	if (!dtDistancePtPolyEdgesSqr(pos, verts, nv, edged, edget))
 	{
 		// Point is outside the polygon, dtClamp to nearest edge.
-		dtReal dmin = DT_REAL_MAX;
-		int imin = -1;
-		for (int i = 0; i < nv; ++i)
+		dtReal dmin = edged[0];
+		int imin = 0;
+		for (int i = 1; i < nv; ++i)
 		{
 			if (edged[i] < dmin)
 			{
@@ -1456,9 +1471,9 @@ void dtNavMesh::closestPointOnPolyInTile(const dtMeshTile* tile, unsigned int ip
 				imin = i;
 			}
 		}
+		CA_ASSUME(imin < nv);
 		const dtReal* va = &verts[imin*3];
 		const dtReal* vb = &verts[((imin+1)%nv)*3];
-		CA_SUPPRESS(6385);
 		dtVlerp(closest, va, vb, edget[imin]);
 	}
 	
@@ -1557,7 +1572,7 @@ dtPolyRef dtNavMesh::findCheapestNearPolyInTile(const dtMeshTile* tile, const dt
 
 	// Get nearby polygons from proximity grid.
 	dtPolyRef polys[128];
-	const bool bExcludeUnwalkable = true;
+	constexpr bool bExcludeUnwalkable = true;
 	int polyCount = queryPolygonsInTile(tile, bmin, bmax, polys, 128, bExcludeUnwalkable);
 
 	// Find nearest polygon amongst the nearby polygons.
@@ -1566,7 +1581,7 @@ dtPolyRef dtNavMesh::findCheapestNearPolyInTile(const dtMeshTile* tile, const dt
 	unsigned char cheapestAreaCostOrder = 0xff;
 	for (int i = 0; i < polyCount; ++i)
 	{
-		dtPolyRef ref = polys[i];
+		const dtPolyRef ref = polys[i];
 		
 		const int polyIdx = decodePolyIdPoly(ref);
 		dtPoly* poly = &tile->polys[polyIdx];
@@ -1582,7 +1597,7 @@ dtPolyRef dtNavMesh::findCheapestNearPolyInTile(const dtMeshTile* tile, const dt
 		{
 			dtReal closestPtPoly[3];
 			closestPointOnPolyInTile(tile, polyIdx, center, closestPtPoly);
-			dtReal d = dtVdistSqr(center, closestPtPoly);
+			const dtReal d = dtVdistSqr(center, closestPtPoly);
 			if (d < nearestDistanceSqr)
 			{
 				dtVcopy(nearestPt, closestPtPoly);
@@ -1592,8 +1607,9 @@ dtPolyRef dtNavMesh::findCheapestNearPolyInTile(const dtMeshTile* tile, const dt
 		}
 	}
 
-	// Verify if the point is actually within requested height, caller is performing 2D check anyway (radius)
-	if (dtAbs(nearestPt[1] - center[1]) > extents[1])
+	// Verify if the point is actually within requested height, caller is performing 2D check anyway (radius).
+	// Using nearest != 0 indicate if nearestPt has been set.
+	if (nearest != 0 && dtAbs(nearestPt[1] - center[1]) > extents[1])
 	{
 		nearest = 0;
 	}
@@ -1857,11 +1873,16 @@ dtStatus dtNavMesh::addTile(unsigned char* data, int dataSize, int flags,
 	// Create connections with neighbour tiles.
 	ReadTilesHelper TileArray;
 	int nneis = 0;
-	dtMeshTile** neis = NULL;
+	dtMeshTile** neis = nullptr;
 
 	// Connect with layers in current tile.
 	nneis = getTileCountAt(header->x, header->y);
 	neis = TileArray.PrepareArray(nneis);
+	if (neis == nullptr)
+	{
+		return DT_FAILURE | DT_OUT_OF_MEMORY;
+	}
+
 	getTilesAt(header->x, header->y, neis, nneis);
 	for (int j = 0; j < nneis; ++j)
 	{
@@ -1880,11 +1901,16 @@ dtStatus dtNavMesh::addTile(unsigned char* data, int dataSize, int flags,
 	// Connect with neighbour tiles.
 	for (int i = 0; i < 8; ++i)
 	{
-		nneis = getNeighbourTilesCountAt(header->x, header->y, i);
-		neis = TileArray.PrepareArray(nneis);
-		getNeighbourTilesAt(header->x, header->y, i, neis, nneis);
+		const int neighbourTileCount = getNeighbourTilesCountAt(header->x, header->y, i);
+		neis = TileArray.PrepareArray(neighbourTileCount );
+		if (neis == nullptr)
+		{
+			return DT_FAILURE | DT_OUT_OF_MEMORY;
+		}
+		
+		getNeighbourTilesAt(header->x, header->y, i, neis, neighbourTileCount);
 
-		for (int j = 0; j < nneis; ++j)
+		for (int j = 0; j < neighbourTileCount; ++j)
 		{
 			// Skip diagonal tiles, nothing to connect there 
 			// (tiles are visited in a ring around the current tile, even tiles are primary directions)
@@ -2205,12 +2231,19 @@ dtStatus dtNavMesh::removeTile(dtTileRef ref, unsigned char** data, int* dataSiz
 	ReadTilesHelper TileArray;
 	int nneis = getTileCountAt(tile->header->x, tile->header->y);
 	dtMeshTile** neis = TileArray.PrepareArray(nneis);
+	if (neis == nullptr)
+	{
+		return DT_FAILURE | DT_OUT_OF_MEMORY;
+	}
 	
 	// Connect with layers in current tile.
 	getTilesAt(tile->header->x, tile->header->y, neis, nneis);
 	for (int j = 0; j < nneis; ++j)
 	{
-		if (neis[j] == tile) continue;
+		CA_ASSUME(j < TileArray.NumAllocated);
+		if (neis[j] == tile)
+			continue;
+
 		unconnectExtLinks(neis[j], tile);
 	}
 	
@@ -2219,10 +2252,17 @@ dtStatus dtNavMesh::removeTile(dtTileRef ref, unsigned char** data, int* dataSiz
 	{
 		nneis = getNeighbourTilesCountAt(tile->header->x, tile->header->y, i);
 		neis = TileArray.PrepareArray(nneis);
+		if (neis == nullptr)
+		{
+			return DT_FAILURE | DT_OUT_OF_MEMORY;
+		}
 
 		getNeighbourTilesAt(tile->header->x, tile->header->y, i, neis, nneis);
 		for (int j = 0; j < nneis; ++j)
+		{
+			CA_ASSUME(j < TileArray.NumAllocated);
 			unconnectExtLinks(neis[j], tile);
+		}
 	}
 
 	// Whether caller wants to own tile data
diff --git a/Engine/Source/Runtime/Navmesh/Public/Detour/DetourNavMesh.h b/Engine/Source/Runtime/Navmesh/Public/Detour/DetourNavMesh.h
--- a/Engine/Source/Runtime/Navmesh/Public/Detour/DetourNavMesh.h
+++ b/Engine/Source/Runtime/Navmesh/Public/Detour/DetourNavMesh.h
@@ -1002,7 +1002,7 @@ struct ReadTilesHelper
 		{
 			dtFree(AllocatedTiles, DT_ALLOC_TEMP);
 			AllocatedTiles = (dtMeshTile**)dtAlloc(RequestedSize * sizeof(dtMeshTile*), DT_ALLOC_TEMP);
-			NumAllocated = RequestedSize;
+			NumAllocated = AllocatedTiles != nullptr ? RequestedSize : 0;
 		}
 
 		return AllocatedTiles;
