Navlink generation optimizations
Use stack allocator to allocate samples
Added profiler events
Addition of a parameter to reduce sampling precision
Added edge index labels to the diplay allowing to inspect a specific edge easily

diff --git a/Engine/Source/Runtime/NavigationSystem/Private/NavMesh/LinkGenerationConfig.cpp b/Engine/Source/Runtime/NavigationSystem/Private/NavMesh/LinkGenerationConfig.cpp
--- a/Engine/Source/Runtime/NavigationSystem/Private/NavMesh/LinkGenerationConfig.cpp
+++ b/Engine/Source/Runtime/NavigationSystem/Private/NavMesh/LinkGenerationConfig.cpp
@@ -11,6 +11,7 @@ void FNavLinkGenerationJumpDownConfig::CopyToDetourConfig(dtNavLinkBuilderJumpDo
 	OutDetourConfig.jumpDistanceFromEdge = JumpDistanceFromEdge;
 	OutDetourConfig.jumpMaxDepth = JumpMaxDepth;
 	OutDetourConfig.jumpEndsHeightTolerance	= JumpEndsHeightTolerance;
+	OutDetourConfig.samplingSeparationFactor = SamplingSeparationFactor;
 }
 
 void FNavLinkGenerationJumpOverConfig::CopyToDetourConfig(dtNavLinkBuilderJumpOverConfig& OutDetourConfig) const
@@ -20,5 +21,6 @@ void FNavLinkGenerationJumpOverConfig::CopyToDetourConfig(dtNavLinkBuilderJumpOv
 	OutDetourConfig.jumpHeight = JumpHeight;
 	OutDetourConfig.jumpHeightTolerance = JumpHeightTolerance;
 	OutDetourConfig.jumpEndsHeightTolerance	= JumpEndsHeightTolerance;
+	OutDetourConfig.samplingSeparationFactor = SamplingSeparationFactor;
 }
 #endif //WITH_RECAST
diff --git a/Engine/Source/Runtime/NavigationSystem/Private/NavMesh/RecastNavMeshGenerator.cpp b/Engine/Source/Runtime/NavigationSystem/Private/NavMesh/RecastNavMeshGenerator.cpp
--- a/Engine/Source/Runtime/NavigationSystem/Private/NavMesh/RecastNavMeshGenerator.cpp
+++ b/Engine/Source/Runtime/NavigationSystem/Private/NavMesh/RecastNavMeshGenerator.cpp
@@ -3695,6 +3695,8 @@ struct FTileGenerationContext
 dtStatus FRecastTileGenerator::BuildTileCacheLinks(FNavMeshBuildContext& BuildContext, dtTileCacheAlloc* alloc, const dtTileCacheLayer& layer,
 	const dtTileCacheContourSet& lcset, TArray<FNavigationLink>& OutGeneratedLinks) const
 {
+	TRACE_CPUPROFILER_EVENT_SCOPE(FRecastTileGenerator::BuildTileCacheLinks);
+	
 	duDebugDraw* dd = nullptr;
 	int32 DebugEdge = -1;
 
@@ -3737,9 +3739,15 @@ dtStatus FRecastTileGenerator::BuildTileCacheLinks(FNavMeshBuildContext& BuildCo
 	if (DebugEdge == -1)
 	{
 		{
-			SCOPE_CYCLE_COUNTER(STAT_Navigation_RecastBuildLinks_Sample);
-			linkBuilder.buildForAllEdges(linkBuilderConfig, DT_LINK_ACTION_JUMP_DOWN);
-			linkBuilder.buildForAllEdges(linkBuilderConfig, DT_LINK_ACTION_JUMP_OVER);
+			{
+				TRACE_CPUPROFILER_EVENT_SCOPE(RecastBuildLinks_JumpDown);
+				linkBuilder.buildForAllEdges(linkBuilderConfig, DT_LINK_ACTION_JUMP_DOWN);
+			}
+
+			{
+				TRACE_CPUPROFILER_EVENT_SCOPE(RecastBuildLinks_JumpOver);
+				linkBuilder.buildForAllEdges(linkBuilderConfig, DT_LINK_ACTION_JUMP_OVER);
+			}
 		}
 
 #if RECAST_INTERNAL_DEBUG_DATA
diff --git a/Engine/Source/Runtime/NavigationSystem/Public/NavMesh/LinkGenerationConfig.h b/Engine/Source/Runtime/NavigationSystem/Public/NavMesh/LinkGenerationConfig.h
--- a/Engine/Source/Runtime/NavigationSystem/Public/NavMesh/LinkGenerationConfig.h
+++ b/Engine/Source/Runtime/NavigationSystem/Public/NavMesh/LinkGenerationConfig.h
@@ -25,7 +25,12 @@ struct FNavLinkGenerationJumpDownConfig
 	float JumpMaxDepth = 150.f;
 
 	UPROPERTY(EditAnywhere, Config, Category = Settings, meta=(Units=cm, UIMin=0, ClampMin=0))
-	float JumpEndsHeightTolerance = 50.f;
+	float JumpEndsHeightTolerance = 80.f;
+
+	/** Value multiplied by CellSize to find the distance between sampling trajectories. Default is 1. */
+    /*  Larger values improve generation speed but might introduce sampling errors.  */
+	UPROPERTY(EditAnywhere, Config, Category = Settings, meta=(UIMin=1, ClampMin=1))
+	float SamplingSeparationFactor = 1.f;
 
 #if WITH_RECAST	
 	/** Copy configuration to dtNavLinkBuilderJumpDownConfig. */
@@ -53,7 +58,12 @@ struct FNavLinkGenerationJumpOverConfig
 	float JumpHeightTolerance = 100.f; 
 
 	UPROPERTY(EditAnywhere, Config, Category = Settings, meta=(Units=cm, UIMin=0, ClampMin=0))
-	float JumpEndsHeightTolerance = 50.f;
+	float JumpEndsHeightTolerance = 80.f;
+
+	/** Value multiplied by CellSize to find the distance between sampling trajectories. Default is 1. */
+    /*  Larger values improve generation speed but might introduce sampling errors.  */
+	UPROPERTY(EditAnywhere, Config, Category = Settings, meta=(UIMin=1, ClampMin=1))
+	float SamplingSeparationFactor = 1.f;
 
 #if WITH_RECAST
 	/** Copy configuration to dtNavLinkBuilderJumpOverConfig. */
diff --git a/Engine/Source/Runtime/Navmesh/Private/DebugUtils/DetourNavLinkDebugDraw.cpp b/Engine/Source/Runtime/Navmesh/Private/DebugUtils/DetourNavLinkDebugDraw.cpp
--- a/Engine/Source/Runtime/Navmesh/Private/DebugUtils/DetourNavLinkDebugDraw.cpp
+++ b/Engine/Source/Runtime/Navmesh/Private/DebugUtils/DetourNavLinkDebugDraw.cpp
@@ -195,10 +195,19 @@ void duDebugDrawNavLinkBuilder(duDebugDraw* dd, const dtNavLinkBuilder& linkBuil
 			dd->begin(DU_DRAW_LINES, 3.0f);
 			for (int i = 0; i < nedges; ++i)
 			{
+				// Display edge index
+				dtReal x = 0.5*(edges[i].sp[0]+edges[i].sq[0]);
+				dtReal y = 0.5*(edges[i].sp[1]+edges[i].sq[1]);
+				dtReal z = 0.5*(edges[i].sp[2]+edges[i].sq[2]);
+				constexpr int bufferSize = 32;
+				char buffer[bufferSize];
+				snprintf(buffer, bufferSize, "%i", i);
+				dd->text(x, y, z, buffer);
+				
 				unsigned int col = duRGBA(0,96,128,255);
 				if (i == selectedEdge)
 					continue;
-				
+
 				dd->vertex(edges[i].sp, col);
 				dd->vertex(edges[i].sq, col);
 			}
@@ -354,7 +363,7 @@ void duDebugDrawNavLinkBuilder(duDebugDraw* dd, const dtNavLinkBuilder& linkBuil
 
 	if (drawFlags & DRAW_SELECTED_EDGE)
 	{
-		if (es)
+		if (es && es->action != DT_LINK_ACTION_UNSET)
 		{
 			dd->begin(DU_DRAW_LINES, 2.0f);
 
diff --git a/Engine/Source/Runtime/Navmesh/Private/Detour/DetourNavLinkBuilder.cpp b/Engine/Source/Runtime/Navmesh/Private/Detour/DetourNavLinkBuilder.cpp
--- a/Engine/Source/Runtime/Navmesh/Private/Detour/DetourNavLinkBuilder.cpp
+++ b/Engine/Source/Runtime/Navmesh/Private/Detour/DetourNavLinkBuilder.cpp
@@ -205,19 +205,6 @@ dtNavLinkBuilder::Trajectory2D::~Trajectory2D()
 }
 
 dtNavLinkBuilder::GroundSegment::~GroundSegment()
-{
-	dtFree(gsamples, DT_ALLOC_PERM_TILE_LINK_BUILDER);
-}
-
-dtNavLinkBuilder::dtNavLinkBuilder() :
-	m_solid(nullptr),
-	m_chf(nullptr),
-	m_edges(nullptr),
-	m_nedges(0),
-	m_links(nullptr),
-	m_nlinks(0),
-	m_clinks(0),
-	m_debugSelectedEdge(-1)
 {
 }
 
@@ -249,9 +236,13 @@ bool dtNavLinkBuilder::findEdges(rcContext& ctx, const rcConfig& cfg, const dtLi
 	m_linkBuilderConfig = builderConfig;
 
 	m_cs = cfg.cs;
+	m_ch = cfg.ch;
+	m_invCs = 1.0/cfg.cs;
 	m_solid = solidHF;
 	m_chf = compactHF;
 
+	dtAssert(m_cs == m_chf->cs && m_ch == m_chf->ch);
+
 	// Build edges.
 	m_nedges = 0;
 	for (int i = 0; i < lcset.nconts; ++i)
@@ -287,7 +278,7 @@ bool dtNavLinkBuilder::findEdges(rcContext& ctx, const rcConfig& cfg, const dtLi
 		{
 			const unsigned short* va = &c.verts[k*4];
 			const unsigned short* vb = &c.verts[j*4];
-			
+
 			// Check k-j for matching contour
 			bool matchFound = false;
 			for (int ii = 0; ii < lcset.nconts; ++ii)
@@ -360,6 +351,8 @@ dtNavLinkBuilder::JumpLink* dtNavLinkBuilder::addLink()
 
 void dtNavLinkBuilder::addEdgeLinks(const dtLinkBuilderConfig& builderConfig, const EdgeSampler* es)
 {
+	TRACE_CPUPROFILER_EVENT_SCOPE(dtNavLinkBuilder::addEdgeLinks);
+	
 	using namespace UE::Detour::NavLink::Private;
 	
 	if (es->start.ngsamples != es->end.ngsamples)
@@ -466,6 +459,8 @@ void dtNavLinkBuilder::addEdgeLinks(const dtLinkBuilderConfig& builderConfig, co
 
 void dtNavLinkBuilder::filterJumpOverLinks() const
 {
+	TRACE_CPUPROFILER_EVENT_SCOPE(dtNavLinkBuilder::filterJumpOverLinks);
+	
 	using namespace UE::Detour::NavLink::Private;
 	
 	// Filter out links which overlap
@@ -545,13 +540,17 @@ void dtNavLinkBuilder::debugBuildEdge(const dtLinkBuilderConfig& builderConfig,
 	filterJumpOverLinks();
 }
 
-bool dtNavLinkBuilder::getCompactHeightfieldHeight(const dtReal* pt, const float hrange, dtReal* height) const
+bool dtNavLinkBuilder::getCompactHeightfieldHeight(const dtReal* pt, const dtReal hrange, dtReal* height) const
 {
-	const float range = m_chf->cs;
-	const int ix0 = dtClamp((int)floorf((pt[0]-range - m_chf->bmin[0])/m_chf->cs), 0, m_chf->width-1);
-	const int iz0 = dtClamp((int)floorf((pt[2]-range - m_chf->bmin[2])/m_chf->cs), 0, m_chf->height-1);
-	const int ix1 = dtClamp((int)floorf((pt[0]+range - m_chf->bmin[0])/m_chf->cs), 0, m_chf->width-1);
-	const int iz1 = dtClamp((int)floorf((pt[2]+range - m_chf->bmin[2])/m_chf->cs), 0, m_chf->height-1);
+
+	const int chfWidth = m_chf->width;
+	const int chfHeight = m_chf->height;
+
+	const dtReal range = m_cs;
+	const int ix0 = dtClamp((int)dtFloor((pt[0]-range - m_chf->bmin[0])*m_invCs), 0, chfWidth-1);
+	const int iz0 = dtClamp((int)dtFloor((pt[2]-range - m_chf->bmin[2])*m_invCs), 0, chfHeight-1);
+	const int ix1 = dtClamp((int)dtFloor((pt[0]+range - m_chf->bmin[0])*m_invCs), 0, chfWidth-1);
+	const int iz1 = dtClamp((int)dtFloor((pt[2]+range - m_chf->bmin[2])*m_invCs), 0, chfHeight-1);
 	
 	dtReal bestDist = DT_REAL_MAX;
 	dtReal bestHeight = DT_REAL_MAX;
@@ -561,16 +560,15 @@ bool dtNavLinkBuilder::getCompactHeightfieldHeight(const dtReal* pt, const float
 	{
 		for (int x = ix0; x <= ix1; ++x)
 		{
-			const rcCompactCell& c = m_chf->cells[x+z*m_chf->width];
-			for (int i = (int)c.index, ni = (int)(c.index+c.count); i < ni; ++i)
+			const rcCompactCell& c = m_chf->cells[x+z*chfWidth];
+			for (unsigned int i = c.index, ni = c.index+c.count; i < ni; ++i)
 			{
-				const rcCompactSpan& s = m_chf->spans[i];
 				if (m_chf->areas[i] == RC_NULL_AREA)
 				{
 					continue;
 				}
 				
-				const dtReal y = m_chf->bmin[1] + s.y * m_chf->ch;
+				const dtReal y = m_chf->bmin[1] + m_chf->spans[i].y * m_ch;
 				const dtReal dist = abs(y - pt[1]);
 				if (dist < hrange && dist < bestDist)
 				{
@@ -650,37 +648,39 @@ bool dtNavLinkBuilder::isTrajectoryClear(const dtReal* pa, const dtReal* pb, con
 	return true;	
 }
 
-void dtNavLinkBuilder::sampleGroundSegment(GroundSegment* seg, const float nsamples, const float groundRange) const
+void dtNavLinkBuilder::sampleGroundSegment(GroundSegment* seg, const int nsamples, const float groundRange) const
 {
+	TRACE_CPUPROFILER_EVENT_SCOPE(dtNavLinkBuilder::sampleGroundSegment);
+	
 	dtReal delta[3];
 	dtVsub(delta, seg->p, seg->q);
 	
 	seg->ngsamples = nsamples;
-	// TODO: profile, preallocate, use dtChunkArray or else.
-	dtAssert(seg->gsamples == nullptr);
-	seg->gsamples = (GroundSample*)dtAlloc(sizeof(GroundSample)*seg->ngsamples, DT_ALLOC_PERM_TILE_LINK_BUILDER);	 
 	seg->npass = 0;
-	
-	for (int i = 0; i < seg->ngsamples; ++i)
+
+	const float invLastIndex = 1.f/(nsamples-1);
+	for (int i = 0; i < nsamples; ++i)
 	{
-		const float u = (float)i/(float)(seg->ngsamples-1);
+		const float u = (float)i*invLastIndex;
 		dtReal pt[3];
 
-		GroundSample* s = &seg->gsamples[i];
+		GroundSample& s = seg->gsamples.Emplace_GetRef();
 		dtVlerp(pt, seg->p, seg->q, u);
-		s->flags = dtNavLinkBuilder::UNSET;
-		if (!getCompactHeightfieldHeight(pt, groundRange, &s->height))
+		s.flags = dtNavLinkBuilder::UNSET;
+		if (!getCompactHeightfieldHeight(pt, groundRange, &s.height))
 		{
 			continue;
 		}
 		
-		s->flags = static_cast<GroundSampleFlag>((unsigned char)s->flags | (unsigned char)HAS_GROUND);
+		s.flags = static_cast<GroundSampleFlag>((unsigned char)s.flags | (unsigned char)HAS_GROUND);
 		seg->npass++;
 	}
 }
 
-void dtNavLinkBuilder::sampleAction(const EdgeSampler* es) const
+void dtNavLinkBuilder::sampleAction(EdgeSampler* es) const
 {
+	TRACE_CPUPROFILER_EVENT_SCOPE(dtNavLinkBuilder::sampleAction);
+	
 	if (es->start.ngsamples != es->end.ngsamples)
 		return;
 	
@@ -688,10 +688,10 @@ void dtNavLinkBuilder::sampleAction(const EdgeSampler* es) const
 	
 	for (int i = 0; i < nsamples; ++i)
 	{
-		GroundSample* ssmp = &es->start.gsamples[i];
-		GroundSample* esmp = &es->end.gsamples[i];
+		GroundSample& ssmp = es->start.gsamples[i];
+		GroundSample& esmp = es->end.gsamples[i];
 		
-		if ((ssmp->flags & HAS_GROUND) == 0 || (esmp->flags & HAS_GROUND) == 0)
+		if ((ssmp.flags & HAS_GROUND) == 0 || (esmp.flags & HAS_GROUND) == 0)
 			continue;
 
 		const dtReal u = (dtReal)i/(dtReal)(nsamples-1);
@@ -699,13 +699,13 @@ void dtNavLinkBuilder::sampleAction(const EdgeSampler* es) const
 		dtVlerp(spt, es->start.p, es->start.q, u);
 		dtVlerp(ept, es->end.p, es->end.q, u);
 		
-		spt[1] = ssmp->height;
-		ept[1] = esmp->height;
+		spt[1] = ssmp.height;
+		ept[1] = esmp.height;
 		
 		if (!isTrajectoryClear(spt, ept, &es->trajectory))
 			continue;
 
-		ssmp->flags = static_cast<GroundSampleFlag>((unsigned char)ssmp->flags | (unsigned char)UNRESTRICTED);
+		ssmp.flags = static_cast<GroundSampleFlag>((unsigned char)ssmp.flags | (unsigned char)UNRESTRICTED);
 	}
 }
 
@@ -713,25 +713,28 @@ void dtNavLinkBuilder::initTrajectory(Trajectory2D* tra) const
 {
 	using namespace UE::Detour::NavLink::Private;
 	
-	const float cs = m_linkBuilderConfig.cellSize;
-	
 	const float* pa = &tra->spine[0];
 	const float* pb = &tra->spine[(tra->nspine-1)*2];
 	
 	const float dx = pb[0] - pa[0];
-	tra->nsamples = dtMax(2, (int)ceilf(dx/cs));
+	const int nsamples  = dtMax(2, (int)ceilf(dx*m_invCs));
+	tra->nsamples = nsamples;
 	// TODO: profile, preallocate, use dtChunkArray or else.
 	dtAssert(tra->samples == nullptr);
-	tra->samples = (TrajectorySample*)dtAlloc(sizeof(TrajectorySample)*tra->nsamples, DT_ALLOC_PERM_TILE_LINK_BUILDER);
-	
-	for (int i = 0; i < tra->nsamples; ++i)
+	tra->samples = (TrajectorySample*)dtAlloc(sizeof(TrajectorySample)*nsamples, DT_ALLOC_PERM_TILE_LINK_BUILDER);
+
+	const unsigned short lastSampleIndex = nsamples-1;
+	const float invLastIndex = 1.f/lastSampleIndex;
+	for (int i = 0; i < nsamples; ++i)
 	{
-		const float u = (float)i / (float)(tra->nsamples-1);
+		const float u = (float)i * invLastIndex;
 		TrajectorySample* s = &tra->samples[i];
 		s->x = dtLerp(pa[0], pb[0], u);
-		
-		const float y0 = getHeight(s->x-m_linkBuilderConfig.agentRadius, tra->spine, tra->nspine);
-		const float y1 = getHeight(s->x+m_linkBuilderConfig.agentRadius, tra->spine, tra->nspine);
+
+		const float* spine = tra->spine;
+		unsigned char nspine = tra->nspine;
+		const float y0 = getHeight(s->x-m_linkBuilderConfig.agentRadius, spine, nspine);
+		const float y1 = getHeight(s->x+m_linkBuilderConfig.agentRadius, spine, nspine);
 
 		const float y = dtLerp(pa[1], pb[1], u);
 		
@@ -745,6 +748,8 @@ int dtNavLinkBuilder::findPotentialJumpOverEdges(const dtReal* sp, const dtReal*
 												  const float depthRange, const float heightRange,
 												  dtReal* outSegs, const int maxOutSegs) const
 {
+	TRACE_CPUPROFILER_EVENT_SCOPE(dtNavLinkBuilder::findPotentialJumpOverEdges);
+	
 	using namespace UE::Detour::NavLink::Private;
 	
 	// Find potential edges to join to.
@@ -967,7 +972,11 @@ void dtNavLinkBuilder::initJumpOverRig(EdgeSampler* es, const dtReal* sp, const
 
 bool dtNavLinkBuilder::sampleEdge(const dtLinkBuilderConfig& builderConfig, dtNavLinkAction desiredAction, const dtReal* sp, const dtReal* sq, dtNavLinkBuilder::EdgeSampler* es) const
 {
+	TRACE_CPUPROFILER_EVENT_SCOPE(dtNavLinkBuilder::sampleEdge);
+	
 	using namespace UE::Detour::NavLink::Private;
+
+	float samplingSeparationFactor = 1.f;
 	
 	if (desiredAction == DT_LINK_ACTION_JUMP_DOWN)
 	{
@@ -976,6 +985,7 @@ bool dtNavLinkBuilder::sampleEdge(const dtLinkBuilderConfig& builderConfig, dtNa
 		const float jumpEndDist = config.jumpLength;
 		const float jumpDownDist = config.jumpMaxDepth;
 		const float groundRange = config.jumpEndsHeightTolerance;
+		samplingSeparationFactor = config.samplingSeparationFactor;
 		initJumpDownRig(es, sp, sq, -jumpStartDist, jumpEndDist, -jumpDownDist, groundRange);
 	}
 	else if (desiredAction == DT_LINK_ACTION_JUMP_OVER)
@@ -1007,6 +1017,7 @@ bool dtNavLinkBuilder::sampleEdge(const dtLinkBuilderConfig& builderConfig, dtNa
 		const float jumpStartDist = config.jumpDistanceFromEdge; 
 		const float jumpHeight = config.jumpHeight;
 		const float groundRange = config.jumpEndsHeightTolerance;
+		samplingSeparationFactor = config.samplingSeparationFactor;
 		initJumpOverRig(es, &segs[ibest*6+0], &segs[ibest*6+3], -jumpStartDist, jumpStartDist, jumpHeight, groundRange);
 	}
 	
@@ -1023,8 +1034,10 @@ bool dtNavLinkBuilder::sampleEdge(const dtLinkBuilderConfig& builderConfig, dtNa
 
 	// Sample start and end ground segments.
 	const float dist = sqrtf(dtVdistSqr(es->rigp, es->rigq));
-	const int ngsamples = dtMax(2, (int)ceilf(dist/m_cs)); // One trajectory per cs
-	
+
+	const dtReal distBetweenSamples = samplingSeparationFactor*m_cs;
+	const int ngsamples = dtMax(2, (int)ceilf(dist/distBetweenSamples));
+
 	sampleGroundSegment(&es->start, ngsamples, es->groundRange);
 	sampleGroundSegment(&es->end, ngsamples, es->groundRange);
 
diff --git a/Engine/Source/Runtime/Navmesh/Public/Detour/DetourNavLinkBuilder.h b/Engine/Source/Runtime/Navmesh/Public/Detour/DetourNavLinkBuilder.h
--- a/Engine/Source/Runtime/Navmesh/Public/Detour/DetourNavLinkBuilder.h
+++ b/Engine/Source/Runtime/Navmesh/Public/Detour/DetourNavLinkBuilder.h
@@ -42,15 +42,17 @@ class dtNavLinkBuilder
 	dtLinkBuilderConfig m_linkBuilderConfig;
 
 	dtReal m_cs = 0;
-	const rcHeightfield* m_solid;
-	const rcCompactHeightfield* m_chf;
+	dtReal m_ch = 0;
+	dtReal m_invCs = 0;
+	const rcHeightfield* m_solid = nullptr;
+	const rcCompactHeightfield* m_chf = nullptr;
 
 	struct Edge
 	{
 		dtReal sp[3], sq[3];
 	};
-	Edge* m_edges;
-	int m_nedges;
+	Edge* m_edges = nullptr;
+	int m_nedges = 0;
 	
 public:	
 	static constexpr int MAX_SPINE = 8;
@@ -97,11 +99,11 @@ private:
 	
 	struct GroundSegment
 	{
-		GroundSegment() : gsamples(nullptr), ngsamples(0) {}
+		GroundSegment() : ngsamples(0) {}
 		NAVMESH_API ~GroundSegment();
 
 		dtReal p[3], q[3];
-		GroundSample* gsamples;
+		TArray<GroundSample, TInlineAllocator<32>> gsamples;
 		unsigned short ngsamples;
 		unsigned short npass;
 	};
@@ -136,17 +138,16 @@ public:
 		JumpLinkFlag flags;
 		dtNavLinkAction action = DT_LINK_ACTION_UNSET;
 	};
-	JumpLink* m_links;
-	int m_nlinks;
-	int m_clinks;
+	JumpLink* m_links = nullptr;
+	int m_nlinks = 0;
+	int m_clinks = 0;
 
 private:
-	int m_debugSelectedEdge;
+	int m_debugSelectedEdge = -1;
 
 	friend NAVMESH_API void duDebugDrawNavLinkBuilder(struct duDebugDraw* dd, const dtNavLinkBuilder& linkBuilder, unsigned int drawFlags, const EdgeSampler* es);
 	
 public:
-	NAVMESH_API dtNavLinkBuilder();
 	NAVMESH_API ~dtNavLinkBuilder();
 
 	// Loops through contours to store edge points in world coordinates.
@@ -176,12 +177,12 @@ private:
 						 const float jumpStartDist, const float jumpEndDist,
 						 const float jumpHeight, const float groundRange);
 
-	bool getCompactHeightfieldHeight(const dtReal* pt, const float hrange, dtReal* height) const;
+	bool getCompactHeightfieldHeight(const dtReal* pt, const dtReal hrange, dtReal* height) const;
 	bool checkHeightfieldCollision(const dtReal x, const dtReal ymin, const dtReal ymax, const dtReal z) const;
 
-	void sampleGroundSegment(GroundSegment* seg, const float nsamples, const float groundRange) const;
+	void sampleGroundSegment(GroundSegment* seg, const int nsamples, const float groundRange) const;
 	
-	void sampleAction(const EdgeSampler* es) const;
+	void sampleAction(EdgeSampler* es) const;
 	
 	void filterJumpOverLinks() const;
 
diff --git a/Engine/Source/Runtime/Navmesh/Public/Detour/DetourNavLinkBuilderConfig.h b/Engine/Source/Runtime/Navmesh/Public/Detour/DetourNavLinkBuilderConfig.h
--- a/Engine/Source/Runtime/Navmesh/Public/Detour/DetourNavLinkBuilderConfig.h
+++ b/Engine/Source/Runtime/Navmesh/Public/Detour/DetourNavLinkBuilderConfig.h
@@ -8,6 +8,7 @@ struct dtNavLinkBuilderJumpDownConfig
 	float jumpDistanceFromEdge = 10.f; 
 	float jumpMaxDepth = 150.f;
 	float jumpEndsHeightTolerance = 50.f;
+	float samplingSeparationFactor = 1.f;	
 };
 
 /** Configuration for generated jump over links. */
@@ -17,5 +18,6 @@ struct dtNavLinkBuilderJumpOverConfig
 	float jumpDistanceFromEdge = 100.f; 
 	float jumpHeight = 100.f; 
 	float jumpHeightTolerance = 100.f; 
-	float jumpEndsHeightTolerance = 50.f; 
+	float jumpEndsHeightTolerance = 50.f;
+	float samplingSeparationFactor = 1.f;
 };
