[NavMesh] NavLink generation
Using inline containers to limit dynamic allocations
Added options to toggle usage of link configs
Added build time in the display stats and added logs

diff --git a/Engine/Source/Runtime/NavigationSystem/Private/NavMesh/LinkGenerationConfig.cpp b/Engine/Source/Runtime/NavigationSystem/Private/NavMesh/LinkGenerationConfig.cpp
--- a/Engine/Source/Runtime/NavigationSystem/Private/NavMesh/LinkGenerationConfig.cpp
+++ b/Engine/Source/Runtime/NavigationSystem/Private/NavMesh/LinkGenerationConfig.cpp
@@ -7,6 +7,7 @@
 
 void FNavLinkGenerationJumpDownConfig::CopyToDetourConfig(dtNavLinkBuilderJumpDownConfig& OutDetourConfig) const
 {
+	OutDetourConfig.enabled = bEnable;
 	OutDetourConfig.jumpLength = JumpLength;
 	OutDetourConfig.jumpDistanceFromEdge = JumpDistanceFromEdge;
 	OutDetourConfig.jumpMaxDepth = JumpMaxDepth;
@@ -16,6 +17,7 @@ void FNavLinkGenerationJumpDownConfig::CopyToDetourConfig(dtNavLinkBuilderJumpDo
 
 void FNavLinkGenerationJumpOverConfig::CopyToDetourConfig(dtNavLinkBuilderJumpOverConfig& OutDetourConfig) const
 {
+	OutDetourConfig.enabled = bEnabled;
 	OutDetourConfig.jumpLength = JumpLength;
 	OutDetourConfig.jumpDistanceFromEdge = JumpDistanceFromEdge;
 	OutDetourConfig.jumpHeight = JumpHeight;
diff --git a/Engine/Source/Runtime/NavigationSystem/Private/NavMesh/NavMeshRenderingComponent.cpp b/Engine/Source/Runtime/NavigationSystem/Private/NavMesh/NavMeshRenderingComponent.cpp
--- a/Engine/Source/Runtime/NavigationSystem/Private/NavMesh/NavMeshRenderingComponent.cpp
+++ b/Engine/Source/Runtime/NavigationSystem/Private/NavMesh/NavMeshRenderingComponent.cpp
@@ -526,6 +526,7 @@ void FNavMeshSceneProxyData::GatherData(const ARecastNavMesh* NavMesh, int32 InN
 		double AverageBuildTime = 0.;
 		double AverageCompLayersBuildTime = 0.;
 		double AverageNavLayersBuildTime = 0.;
+		double TotalBuildLinkTime = 0.;
 		if (bGatherTileBuildTimes || NavMeshGeometry.bGatherTileBuildTimesHeatMap)
 		{
 			QUICK_SCOPE_CYCLE_COUNTER(STAT_NavMesh_GatherDebugDrawing_MaxTileBuildTime);
@@ -543,6 +544,7 @@ void FNavMeshSceneProxyData::GatherData(const ARecastNavMesh* NavMesh, int32 InN
 					AverageBuildTime += Pair.Value.BuildTime;
 					AverageCompLayersBuildTime += Pair.Value.BuildCompressedLayerTime;
 					AverageNavLayersBuildTime += Pair.Value.BuildNavigationDataTime;
+					TotalBuildLinkTime += Pair.Value.BuildLinkTime;
 				}
 			}
 
@@ -632,6 +634,7 @@ void FNavMeshSceneProxyData::GatherData(const ARecastNavMesh* NavMesh, int32 InN
 				DebugLabels.Add(FDebugText(FString::Printf(TEXT("   Avg nav layers time: %0.1f ms"), AverageNavLayersBuildTime*1000.)));
 				DebugLabels.Add(FDebugText(FString::Printf(TEXT("Min: %0.2f ms  Max: %0.2f ms"), NavMeshGeometry.MinTileBuildTime*1000., NavMeshGeometry.MaxTileBuildTime*1000.)));
 				DebugLabels.Add(FDebugText(FString::Printf(TEXT("Total: %0.3f s"), TotalTileBuildTime)));
+				DebugLabels.Add(FDebugText(FString::Printf(TEXT("   Total link build time: %.4f s"), TotalBuildLinkTime)));
 
 				DebugLabels.Add(FDebugText(TEXT(""))); // empty line
 				const double TileAreaM2 = FMath::Square(NavMesh->TileSizeUU) / 10000.;
diff --git a/Engine/Source/Runtime/NavigationSystem/Private/NavMesh/RecastNavMeshGenerator.cpp b/Engine/Source/Runtime/NavigationSystem/Private/NavMesh/RecastNavMeshGenerator.cpp
--- a/Engine/Source/Runtime/NavigationSystem/Private/NavMesh/RecastNavMeshGenerator.cpp
+++ b/Engine/Source/Runtime/NavigationSystem/Private/NavMesh/RecastNavMeshGenerator.cpp
@@ -3696,6 +3696,8 @@ dtStatus FRecastTileGenerator::BuildTileCacheLinks(FNavMeshBuildContext& BuildCo
 	const dtTileCacheContourSet& lcset, TArray<FNavigationLink>& OutGeneratedLinks) const
 {
 	TRACE_CPUPROFILER_EVENT_SCOPE(FRecastTileGenerator::BuildTileCacheLinks);
+	BuildContext.log(RC_LOG_PROGRESS, "Building Links:");
+	const double StartTime = FPlatformTime::Seconds();
 	
 	duDebugDraw* dd = nullptr;
 	int32 DebugEdge = -1;
@@ -3711,8 +3713,18 @@ dtStatus FRecastTileGenerator::BuildTileCacheLinks(FNavMeshBuildContext& BuildCo
 
 	dtAssert(alloc);
 
+	auto LogOnExit = [&]()
+	{
+#if RECAST_INTERNAL_DEBUG_DATA
+		const double LinkBuildTime = FPlatformTime::Seconds()-StartTime;
+		BuildContext.InternalDebugData.BuildLinkTime += LinkBuildTime;
+#endif
+		BuildContext.log(RC_LOG_PROGRESS, "   BuildTileCacheLinks time: %0.3fms.", LinkBuildTime*1000 );
+	};
+
 	if (!SolidHF || !CompactHF)
 	{
+		LogOnExit();
 		return DT_FAILURE;
 	}
 
@@ -3732,22 +3744,32 @@ dtStatus FRecastTileGenerator::BuildTileCacheLinks(FNavMeshBuildContext& BuildCo
 
 		if (!linkBuilder.findEdges(BuildContext, TileConfig, linkBuilderConfig, lcset, orig, SolidHF, CompactHF))
 		{
+			LogOnExit();
 			return DT_FAILURE;
 		}
+
+		BuildContext.log(RC_LOG_PROGRESS, "   Found %i edges.", linkBuilder.getEdgeCount());
 	}
 
+	if (linkBuilder.getEdgeCount() == 0)
+	{
+		LogOnExit();
+		return DT_SUCCESS;	
+	}
+	
 	if (DebugEdge == -1)
 	{
+		rcContext& context = BuildContext;
+		if (linkBuilderConfig.jumpDownConfig.enabled)
 		{
-			{
-				TRACE_CPUPROFILER_EVENT_SCOPE(RecastBuildLinks_JumpDown);
-				linkBuilder.buildForAllEdges(linkBuilderConfig, DT_LINK_ACTION_JUMP_DOWN);
-			}
+			TRACE_CPUPROFILER_EVENT_SCOPE(RecastBuildLinks_JumpDown);
+			linkBuilder.buildForAllEdges(context, linkBuilderConfig, DT_LINK_ACTION_JUMP_DOWN);
+		}
 
-			{
-				TRACE_CPUPROFILER_EVENT_SCOPE(RecastBuildLinks_JumpOver);
-				linkBuilder.buildForAllEdges(linkBuilderConfig, DT_LINK_ACTION_JUMP_OVER);
-			}
+		if (linkBuilderConfig.jumpOverConfig.enabled)
+		{
+			TRACE_CPUPROFILER_EVENT_SCOPE(RecastBuildLinks_JumpOver);
+			linkBuilder.buildForAllEdges(context, linkBuilderConfig, DT_LINK_ACTION_JUMP_OVER);
 		}
 
 #if RECAST_INTERNAL_DEBUG_DATA
@@ -3755,32 +3777,37 @@ dtStatus FRecastTileGenerator::BuildTileCacheLinks(FNavMeshBuildContext& BuildCo
 	}
 	else
 	{
-		dtNavLinkBuilder::EdgeSampler sampler1;
-		linkBuilder.debugBuildEdge(linkBuilderConfig, DT_LINK_ACTION_JUMP_DOWN, DebugEdge, sampler1);
-		duDebugDrawNavLinkBuilder(dd, linkBuilder, DebugFlags, &sampler1);
+		if (linkBuilderConfig.jumpDownConfig.enabled)
+		{
+			dtNavLinkBuilder::EdgeSampler sampler1;
+			linkBuilder.debugBuildEdge(linkBuilderConfig, DT_LINK_ACTION_JUMP_DOWN, DebugEdge, sampler1);
+			duDebugDrawNavLinkBuilder(dd, linkBuilder, DebugFlags, &sampler1);
+		}
 		
-		dtNavLinkBuilder::EdgeSampler sampler2;
-		linkBuilder.debugBuildEdge(linkBuilderConfig, DT_LINK_ACTION_JUMP_OVER, DebugEdge, sampler2);
-		duDebugDrawNavLinkBuilder(dd, linkBuilder, DebugFlags, &sampler2);
+		if (linkBuilderConfig.jumpOverConfig.enabled)
+		{
+			dtNavLinkBuilder::EdgeSampler sampler2;
+			linkBuilder.debugBuildEdge(linkBuilderConfig, DT_LINK_ACTION_JUMP_OVER, DebugEdge, sampler2);
+			duDebugDrawNavLinkBuilder(dd, linkBuilder, DebugFlags, &sampler2);
+		}
 #endif // RECAST_INTERNAL_DEBUG_DATA		
 	}
 
 	// Make FNavigationLinks
-	for (int i = 0; i < linkBuilder.m_nlinks; ++i)
+	for (const dtNavLinkBuilder::JumpLink& link : linkBuilder.m_links)
 	{
-		const dtNavLinkBuilder::JumpLink* link = &linkBuilder.m_links[i];
-		
-		if (link->flags == dtNavLinkBuilder::INVALID)
+		if (link.flags == dtNavLinkBuilder::INVALID)
 			continue;
 
 		// For now, just make a link using the center of the range.
 		dtReal midA[3];
-		dtVlerp(midA, &link->spine0[0], &link->spine1[0], 0.5);
+		dtVlerp(midA, &link.spine0[0], &link.spine1[0], 0.5);
 		dtReal midB[3];
-		dtVlerp(midB, &link->spine0[(link->nspine-1)*3], &link->spine1[(link->nspine-1)*3], 0.5);
+		dtVlerp(midB, &link.spine0[(link.nspine-1)*3], &link.spine1[(link.nspine-1)*3], 0.5);
 		OutGeneratedLinks.Add(FNavigationLink(Recast2UnrealPoint(midA), Recast2UnrealPoint(midB)));
 	}
 	
+	LogOnExit();
 	return DT_SUCCESS;
 }
 
diff --git a/Engine/Source/Runtime/NavigationSystem/Public/NavMesh/LinkGenerationConfig.h b/Engine/Source/Runtime/NavigationSystem/Public/NavMesh/LinkGenerationConfig.h
--- a/Engine/Source/Runtime/NavigationSystem/Public/NavMesh/LinkGenerationConfig.h
+++ b/Engine/Source/Runtime/NavigationSystem/Public/NavMesh/LinkGenerationConfig.h
@@ -14,6 +14,10 @@ struct FNavLinkGenerationJumpDownConfig
 {
 	GENERATED_USTRUCT_BODY()
 
+	/** Should this config be used to generate links. */
+	UPROPERTY(EditAnywhere, Config, Category = Settings)
+	bool bEnable = true;
+	
 	// @todo: Rename, describe and find best defaults for those parameters.
 	UPROPERTY(EditAnywhere, Config, Category = Settings, meta=(Units=cm, UIMin=0, ClampMin=0))
 	float JumpLength = 150.f; 
@@ -44,6 +48,10 @@ struct FNavLinkGenerationJumpOverConfig
 {
 	GENERATED_USTRUCT_BODY()
 
+	/** Should this config be used to generate links. */
+	UPROPERTY(EditAnywhere, Config, Category = Settings)
+	bool bEnabled = true;
+	
 	// @todo: Rename, describe and find best defaults for those parameters.
 	UPROPERTY(EditAnywhere, Config, Category = Settings, meta=(Units=cm, UIMin=0, ClampMin=0))
 	float JumpLength = 200.f; 
diff --git a/Engine/Source/Runtime/NavigationSystem/Public/NavMesh/RecastInternalDebugData.h b/Engine/Source/Runtime/NavigationSystem/Public/NavMesh/RecastInternalDebugData.h
--- a/Engine/Source/Runtime/NavigationSystem/Public/NavMesh/RecastInternalDebugData.h
+++ b/Engine/Source/Runtime/NavigationSystem/Public/NavMesh/RecastInternalDebugData.h
@@ -33,6 +33,7 @@ struct FRecastInternalDebugData : public duDebugDraw
 	double BuildTime = 0.;
 	double BuildCompressedLayerTime = 0.;
 	double BuildNavigationDataTime = 0.;
+	double BuildLinkTime = 0;
 
 	uint32 TriangleCount = 0;
 	unsigned char Resolution = 0;
diff --git a/Engine/Source/Runtime/Navmesh/Private/DebugUtils/DetourNavLinkDebugDraw.cpp b/Engine/Source/Runtime/Navmesh/Private/DebugUtils/DetourNavLinkDebugDraw.cpp
--- a/Engine/Source/Runtime/Navmesh/Private/DebugUtils/DetourNavLinkDebugDraw.cpp
+++ b/Engine/Source/Runtime/Navmesh/Private/DebugUtils/DetourNavLinkDebugDraw.cpp
@@ -111,14 +111,15 @@ namespace UE::Detour::Private
 		
 		dd->begin(DU_DRAW_QUADS);
 		dtReal p0[3], p1[3], p2[3], p3[3];
-		for (int i = 0; i < trajectory->nsamples; ++i)
+		const int nsamples = trajectory->samples.Num();
+		for (int i = 0; i < nsamples; ++i)
 		{
-			const dtNavLinkBuilder::TrajectorySample* s = &trajectory->samples[i];
-			const dtReal u = (float)i / (float)(trajectory->nsamples-1);
+			const dtNavLinkBuilder::TrajectorySample& s = trajectory->samples[i];
+			const dtReal u = (float)i / (float)(nsamples-1);
 			dtVlerp(p1, pa, pb, u);
 			dtVcopy(p0, p1);
-			p0[1] += s->ymin;
-			p1[1] += s->ymax;
+			p0[1] += s.ymin;
+			p1[1] += s.ymax;
 			
 			if (i > 0)
 			{
@@ -142,16 +143,16 @@ namespace UE::Detour::Private
 		colt = duDarkenCol(colt);
 		
 		dd->begin(DU_DRAW_LINES, 2.0f);
-		for (int i = 0; i < trajectory->nsamples; ++i)
+		for (int i = 0; i < nsamples; ++i)
 		{
-			const dtNavLinkBuilder::TrajectorySample* s = &trajectory->samples[i];
-			const float u = (float)i / (float)(trajectory->nsamples-1);
+			const dtNavLinkBuilder::TrajectorySample& s = trajectory->samples[i];
+			const float u = (float)i / (float)(nsamples-1);
 			dtVlerp(p1, pa, pb, u);
 			dtVcopy(p0, p1);
-			p0[1] += s->ymin;
-			p1[1] += s->ymax;
+			p0[1] += s.ymin;
+			p1[1] += s.ymax;
 			
-			if (i == 0 || i == (trajectory->nsamples-1))
+			if (i == 0 || i == (nsamples-1))
 			{
 				dd->vertex(p0, colb);
 				dd->vertex(p1, colt);
@@ -186,8 +187,8 @@ void duDebugDrawNavLinkBuilder(duDebugDraw* dd, const dtNavLinkBuilder& linkBuil
 
 	if (drawFlags & DRAW_BORDERS)
 	{
-		const dtNavLinkBuilder::Edge* edges = linkBuilder.m_edges;
-		const int nedges = linkBuilder.m_nedges;
+		const TArray<dtNavLinkBuilder::Edge, TInlineAllocator<32>>& edges = linkBuilder.m_edges;
+		const int nedges = linkBuilder.m_edges.Num();
 		const int selectedEdge = linkBuilder.m_debugSelectedEdge;
 		
 		if (nedges)
@@ -251,8 +252,6 @@ void duDebugDrawNavLinkBuilder(duDebugDraw* dd, const dtNavLinkBuilder& linkBuil
 	
 	if (drawFlags & DRAW_LINKS)
 	{
-		const dtNavLinkBuilder::JumpLink* links = linkBuilder.m_links;
-
 		unsigned int jumpDownCol0 = duLerpCol(duColor::blue, duColor::white, 200);
 		unsigned int jumpDownCol1 = duColor::blue;
 		unsigned int jumpOverCol0 = duLerpCol(duColor::lightGrey, duColor::white, 200);
@@ -275,87 +274,82 @@ void duDebugDrawNavLinkBuilder(duDebugDraw* dd, const dtNavLinkBuilder& linkBuil
 		unsigned int col0 = duColor::black;
 		unsigned int col1 = duColor::black;
 		
-		const int nlinks = linkBuilder.m_nlinks;
+		const int nlinks = linkBuilder.m_links.Num();
 		if (nlinks)
 		{
 			dd->begin(DU_DRAW_QUADS);
-			for (int i = 0; i < nlinks; ++i)
+			for (const dtNavLinkBuilder::JumpLink& link : linkBuilder.m_links)
 			{
-				const dtNavLinkBuilder::JumpLink* link = &links[i];
-				
-				if (link->flags == dtNavLinkBuilder::INVALID)
+				if (link.flags == dtNavLinkBuilder::INVALID)
 					continue;
 
-				selectColors(link->action, col0, col1);
+				selectColors(link.action, col0, col1);
 				
-				for (int j = 0; j < link->nspine-1; ++j)
+				for (int j = 0; j < link.nspine-1; ++j)
 				{
-					int u = (j*255)/link->nspine;
+					int u = (j*255)/link.nspine;
 					unsigned int col = duTransCol(duLerpCol(col0,col1,u),128);
-					if (link->flags == dtNavLinkBuilder::INVALID)
+					if (link.flags == dtNavLinkBuilder::INVALID)
 						col = duRGBA(255,0,0,64);
 					
-					dd->vertex(&link->spine1[j*3], col);
-					dd->vertex(&link->spine1[(j+1)*3], col);
-					dd->vertex(&link->spine0[(j+1)*3], col);
-					dd->vertex(&link->spine0[j*3], col);
+					dd->vertex(&link.spine1[j*3], col);
+					dd->vertex(&link.spine1[(j+1)*3], col);
+					dd->vertex(&link.spine0[(j+1)*3], col);
+					dd->vertex(&link.spine0[j*3], col);
 				}
 			}
 			dd->end();
 			
 			dd->begin(DU_DRAW_LINES, 3.0f);
-			for (int i = 0; i < nlinks; ++i)
+			for (const dtNavLinkBuilder::JumpLink& link : linkBuilder.m_links)
 			{
-				const dtNavLinkBuilder::JumpLink* link = &links[i];
-				if (link->flags == dtNavLinkBuilder::INVALID)
+				if (link.flags == dtNavLinkBuilder::INVALID)
 					continue;
 
-				selectColors(link->action, col0, col1);
+				selectColors(link.action, col0, col1);
 				
-				for (int j = 0; j < link->nspine-1; ++j)
+				for (int j = 0; j < link.nspine-1; ++j)
 				{
 					unsigned int col = duTransCol(duDarkenCol(col1),128);
 					
-					dd->vertex(&link->spine0[j*3], col);
-					dd->vertex(&link->spine0[(j+1)*3], col);
-					dd->vertex(&link->spine1[j*3], col);
-					dd->vertex(&link->spine1[(j+1)*3], col);
+					dd->vertex(&link.spine0[j*3], col);
+					dd->vertex(&link.spine0[(j+1)*3], col);
+					dd->vertex(&link.spine1[j*3], col);
+					dd->vertex(&link.spine1[(j+1)*3], col);
 				}
 
-				dd->vertex(&link->spine0[0], duDarkenCol(col1));
-				dd->vertex(&link->spine1[0], duDarkenCol(col1));
+				dd->vertex(&link.spine0[0], duDarkenCol(col1));
+				dd->vertex(&link.spine1[0], duDarkenCol(col1));
 
-				dd->vertex(&link->spine0[(link->nspine-1)*3], duDarkenCol(col1));
-				dd->vertex(&link->spine1[(link->nspine-1)*3], duDarkenCol(col1));
+				dd->vertex(&link.spine0[(link.nspine-1)*3], duDarkenCol(col1));
+				dd->vertex(&link.spine1[(link.nspine-1)*3], duDarkenCol(col1));
 			}
 			dd->end();
 
 			dd->begin(DU_DRAW_POINTS, 8.0f);
-			for (int i = 0; i < nlinks; ++i)
+			for (const dtNavLinkBuilder::JumpLink& link : linkBuilder.m_links)
 			{
-				const dtNavLinkBuilder::JumpLink* link = &links[i];
-				if (link->flags == dtNavLinkBuilder::INVALID)
+				if (link.flags == dtNavLinkBuilder::INVALID)
 					continue;
 
-				selectColors(link->action, col0, col1);
+				selectColors(link.action, col0, col1);
 				
-				dd->vertex(&link->spine0[0], duDarkenCol(col1));
-				dd->vertex(&link->spine1[0], duDarkenCol(col1));
-				dd->vertex(&link->spine0[(link->nspine-1)*3], duDarkenCol(col1));
-				dd->vertex(&link->spine1[(link->nspine-1)*3], duDarkenCol(col1));
+				dd->vertex(&link.spine0[0], duDarkenCol(col1));
+				dd->vertex(&link.spine1[0], duDarkenCol(col1));
+				dd->vertex(&link.spine0[(link.nspine-1)*3], duDarkenCol(col1));
+				dd->vertex(&link.spine1[(link.nspine-1)*3], duDarkenCol(col1));
 			}
 			dd->end();
 			
 			dd->begin(DU_DRAW_POINTS, 4.0f);
-			for (int i = 0; i < nlinks; ++i)
+			for (const dtNavLinkBuilder::JumpLink& link : linkBuilder.m_links)
 			{
-				const dtNavLinkBuilder::JumpLink* link = &links[i];
-				if (link->flags == dtNavLinkBuilder::INVALID)
+				if (link.flags == dtNavLinkBuilder::INVALID)
 					continue;
-				dd->vertex(&link->spine0[0], duColor::lightGrey);
-				dd->vertex(&link->spine1[0], duColor::lightGrey);
-				dd->vertex(&link->spine0[(link->nspine-1)*3], duColor::lightGrey);
-				dd->vertex(&link->spine1[(link->nspine-1)*3], duColor::lightGrey);
+				dd->vertex(&link.spine0[0], duColor::lightGrey);
+				dd->vertex(&link.spine1[0], duColor::lightGrey);
+				dd->vertex(&link.spine0[(link.nspine-1)*3], duColor::lightGrey);
+				dd->vertex(&link.spine1[(link.nspine-1)*3], duColor::lightGrey);
 			}
 			dd->end();
 		}
diff --git a/Engine/Source/Runtime/Navmesh/Private/Detour/DetourNavLinkBuilder.cpp b/Engine/Source/Runtime/Navmesh/Private/Detour/DetourNavLinkBuilder.cpp
--- a/Engine/Source/Runtime/Navmesh/Private/Detour/DetourNavLinkBuilder.cpp
+++ b/Engine/Source/Runtime/Navmesh/Private/Detour/DetourNavLinkBuilder.cpp
@@ -199,40 +199,15 @@ namespace UE::Detour::NavLink::Private
 	}
 }
 
-dtNavLinkBuilder::Trajectory2D::~Trajectory2D()
-{
-	dtFree(samples, DT_ALLOC_PERM_TILE_LINK_BUILDER);
-}
-
 dtNavLinkBuilder::GroundSegment::~GroundSegment()
 {
 }
 
-dtNavLinkBuilder::~dtNavLinkBuilder()
-{	
-	cleanup();
-}
-
-void dtNavLinkBuilder::cleanup()
-{
-	m_solid = nullptr;
-	m_chf = nullptr;
-
-	dtFree(m_edges, DT_ALLOC_PERM_TILE_LINK_BUILDER);
-	m_edges = nullptr;
-	m_nedges = 0;
-	
-	dtFree(m_links, DT_ALLOC_PERM_TILE_LINK_BUILDER);
-	m_links = nullptr;
-	m_nlinks = 0;
-	m_clinks = 0;
-}
-
 bool dtNavLinkBuilder::findEdges(rcContext& ctx, const rcConfig& cfg, const dtLinkBuilderConfig& builderConfig,
                                  const dtTileCacheContourSet& lcset, const dtReal* orig,
                                  const rcHeightfield* solidHF, const rcCompactHeightfield* compactHF)
 {
-	cleanup();
+	dtAssert(m_solid == nullptr && m_chf == nullptr && m_edges.IsEmpty() && m_links.IsEmpty());
 	m_linkBuilderConfig = builderConfig;
 
 	m_cs = cfg.cs;
@@ -244,30 +219,23 @@ bool dtNavLinkBuilder::findEdges(rcContext& ctx, const rcConfig& cfg, const dtLi
 	dtAssert(m_cs == m_chf->cs && m_ch == m_chf->ch);
 
 	// Build edges.
-	m_nedges = 0;
+	int edgeCount = 0;
 	for (int i = 0; i < lcset.nconts; ++i)
 	{
-		m_nedges += lcset.conts[i].nverts;
+		edgeCount += lcset.conts[i].nverts;
 	}
 
-	if (!m_nedges)
+	if (edgeCount == 0)
 	{
 		ctx.log(RC_LOG_ERROR, "fillEdges: No edges!");
 		return false;
 	}
 
-	dtAssert(m_edges == nullptr);
-	m_edges = (Edge*)dtAlloc(sizeof(Edge)*m_nedges, DT_ALLOC_PERM_TILE_LINK_BUILDER);
-	if (!m_edges)
-	{
-		ctx.log(RC_LOG_ERROR, "fillEdges: Failed to alloc edges (%d)", m_nedges);
-		return false;
-	}
+	m_edges.Reserve(edgeCount);
 
 	const dtReal cs = cfg.cs;
 	const dtReal ch = cfg.ch;
 	
-	m_nedges = 0;
 	for (int i = 0; i < lcset.nconts; ++i)
 	{
 		const dtTileCacheContour& c = lcset.conts[i];
@@ -313,15 +281,15 @@ bool dtNavLinkBuilder::findEdges(rcContext& ctx, const rcConfig& cfg, const dtLi
 			if (!matchFound)
 			{
 				// Add edge
-				Edge* e = &m_edges[m_nedges++];
+				Edge& e = m_edges.Emplace_GetRef();
 
-				e->sp[0] = orig[0] + vb[0]*cs;
-				e->sp[1] = orig[1] + (vb[1]+2)*ch;
-				e->sp[2] = orig[2] + vb[2]*cs;
+				e.sp[0] = orig[0] + vb[0]*cs;
+				e.sp[1] = orig[1] + (vb[1]+2)*ch;
+				e.sp[2] = orig[2] + vb[2]*cs;
 
-				e->sq[0] = orig[0] + va[0]*cs;
-				e->sq[1] = orig[1] + (va[1]+2)*ch;
-				e->sq[2] = orig[2] + va[2]*cs;
+				e.sq[0] = orig[0] + va[0]*cs;
+				e.sq[1] = orig[1] + (va[1]+2)*ch;
+				e.sq[2] = orig[2] + va[2]*cs;
 			}
 		}
 	}
@@ -329,26 +297,6 @@ bool dtNavLinkBuilder::findEdges(rcContext& ctx, const rcConfig& cfg, const dtLi
 	return true;
 }
 
-dtNavLinkBuilder::JumpLink* dtNavLinkBuilder::addLink()
-{
-	if (m_nlinks+1 > m_clinks)
-	{
-		m_clinks = m_clinks ? m_clinks*2 : 8;
-
-		JumpLink* n = (JumpLink*)dtAlloc(sizeof(JumpLink)*m_clinks, DT_ALLOC_PERM_TILE_LINK_BUILDER);
-
-		if (m_nlinks)
-		{
-			memcpy(n, m_links, sizeof(JumpLink)*m_nlinks);
-		}
-		dtFree(m_links, DT_ALLOC_PERM_TILE_LINK_BUILDER);
-		m_links = n;
-	}
-	
-	JumpLink* link = &m_links[m_nlinks++];
-	return link;
-}
-
 void dtNavLinkBuilder::addEdgeLinks(const dtLinkBuilderConfig& builderConfig, const EdgeSampler* es)
 {
 	TRACE_CPUPROFILER_EVENT_SCOPE(dtNavLinkBuilder::addEdgeLinks);
@@ -366,7 +314,8 @@ void dtNavLinkBuilder::addEdgeLinks(const dtLinkBuilderConfig& builderConfig, co
 	constexpr int RAD = 2;
 	GroundSampleFlag kernel[RAD*2+1];
 
-	GroundSampleFlag* nflags = (GroundSampleFlag*)dtAlloc(sizeof(GroundSampleFlag)*nsamples, DT_ALLOC_TEMP);
+	TArray<GroundSampleFlag, TInlineAllocator<64>> groundSampleFlags;
+	groundSampleFlags.Reserve(nsamples);
 	
 	for (int i = 0; i < nsamples; ++i)
 	{
@@ -378,7 +327,7 @@ void dtNavLinkBuilder::addEdgeLinks(const dtLinkBuilderConfig& builderConfig, co
 			kernel[nkernel++] = (GroundSampleFlag)(es->start.gsamples[i].flags & UNRESTRICTED);
 		}
 		insertSort((unsigned char*)kernel, nkernel);
-		nflags[i] = kernel[(nkernel+1)/2];
+		groundSampleFlags.Add(kernel[(nkernel+1)/2]);
 	}
 
 	const int agentRadiusVx = builderConfig.agentRadius / builderConfig.cellSize;
@@ -387,7 +336,7 @@ void dtNavLinkBuilder::addEdgeLinks(const dtLinkBuilderConfig& builderConfig, co
 	int start = -1;
 	for (int i = 0; i <= nsamples; ++i)
 	{
-		const bool valid = i < nsamples && nflags[i] != UNSET;
+		const bool valid = i < nsamples && groundSampleFlags[i] != UNSET;
 		if (start == -1)
 		{
 			if (valid)
@@ -413,11 +362,11 @@ void dtNavLinkBuilder::addEdgeLinks(const dtLinkBuilderConfig& builderConfig, co
 					ep[1] = es->end.gsamples[start].height;
 					eq[1] = es->end.gsamples[i-1].height;
 
-					JumpLink* link = addLink();
+					JumpLink& link = m_links.Emplace_GetRef();
 
-					link->action = es->action;
-					link->flags = VALID;
-					link->nspine = es->trajectory.nspine;
+					link.action = es->action;
+					link.flags = VALID;
+					link.nspine = es->trajectory.nspine;
 
 					const float startx = es->trajectory.spine[0];
 					const float endx = es->trajectory.spine[(es->trajectory.nspine-1)*2];
@@ -432,7 +381,7 @@ void dtNavLinkBuilder::addEdgeLinks(const dtLinkBuilderConfig& builderConfig, co
 						const float* spt = &es->trajectory.spine[j*2];
 						const float u = (spt[0] - startx)/deltax;
 						const float dy = spt[1] - dtLerp(starty, endy, u) + m_linkBuilderConfig.agentClimb;
-						dtReal* p = &link->spine0[j*3];
+						dtReal* p = &link.spine0[j*3];
 						dtVlerp(p, sp, ep, u);
 						dtVmad(p, p, es->ay, dy);
 					}
@@ -443,7 +392,7 @@ void dtNavLinkBuilder::addEdgeLinks(const dtLinkBuilderConfig& builderConfig, co
 						const float* spt = &es->trajectory.spine[j*2];
 						const float u = (spt[0] - startx)/deltax;
 						const float dy = spt[1] - dtLerp(starty, endy, u) + m_linkBuilderConfig.agentClimb;
-						dtReal* p = &link->spine1[j*3];
+						dtReal* p = &link.spine1[j*3];
 						dtVlerp(p, sq, eq, u);
 						dtVmad(p, p, es->ay, dy);
 					}
@@ -453,11 +402,9 @@ void dtNavLinkBuilder::addEdgeLinks(const dtLinkBuilderConfig& builderConfig, co
 			}
 		}
 	}
-
-	dtFree(nflags, DT_ALLOC_TEMP);
 }
 
-void dtNavLinkBuilder::filterJumpOverLinks() const
+void dtNavLinkBuilder::filterJumpOverLinks()
 {
 	TRACE_CPUPROFILER_EVENT_SCOPE(dtNavLinkBuilder::filterJumpOverLinks);
 	
@@ -466,27 +413,27 @@ void dtNavLinkBuilder::filterJumpOverLinks() const
 	// Filter out links which overlap
 	const float thresholdSquared = dtSqr(m_chf->cs*5.0f);
 	
-	for (int i = 0; i < m_nlinks-1; ++i)
+	for (int i = 0; i < m_links.Num()-1; ++i)
 	{
-		JumpLink* li = &m_links[i];
-		if (li->flags == INVALID)
+		JumpLink& li = m_links[i];
+		if (li.flags == INVALID)
 			continue;
 		
-		const dtReal* spi = &li->spine0[0];
-		const dtReal* sqi = &li->spine1[0];
-		const dtReal* epi = &li->spine0[(li->nspine-1)*3];
-		const dtReal* eqi = &li->spine1[(li->nspine-1)*3];
+		const dtReal* spi = &li.spine0[0];
+		const dtReal* sqi = &li.spine1[0];
+		const dtReal* epi = &li.spine0[(li.nspine-1)*3];
+		const dtReal* eqi = &li.spine1[(li.nspine-1)*3];
 		
-		for (int j = i+1; j < m_nlinks; ++j)
+		for (int j = i+1; j < m_links.Num(); ++j)
 		{
-			JumpLink* lj = &m_links[j];
-			if (lj->flags == INVALID)
+			JumpLink& lj = m_links[j];
+			if (lj.flags == INVALID)
 				continue;
 			
-			const dtReal* spj = &lj->spine0[0];
-			const dtReal* sqj = &lj->spine1[0];
-			const dtReal* epj = &lj->spine0[(lj->nspine-1)*3];
-			const dtReal* eqj = &lj->spine1[(lj->nspine-1)*3];
+			const dtReal* spj = &lj.spine0[0];
+			const dtReal* sqj = &lj.spine1[0];
+			const dtReal* epj = &lj.spine0[(lj.nspine-1)*3];
+			const dtReal* eqj = &lj.spine1[(lj.nspine-1)*3];
 			
 			const dtReal d0 = distSegSegSqr(spi,sqi, epj,eqj);
 			const dtReal d1 = distSegSegSqr(epi,eqi, spj,sqj);
@@ -495,11 +442,11 @@ void dtNavLinkBuilder::filterJumpOverLinks() const
 			{
 				if (dtVdistSqr(spi,sqi) > dtVdistSqr(spj,sqj))
 				{
-					lj->flags = INVALID;
+					lj.flags = INVALID;
 				}
 				else
 				{
-					li->flags = INVALID;
+					li.flags = INVALID;
 					break;
 				}
 			}
@@ -507,9 +454,9 @@ void dtNavLinkBuilder::filterJumpOverLinks() const
 	}
 }
 
-void dtNavLinkBuilder::buildForAllEdges(const dtLinkBuilderConfig& builderConfig, dtNavLinkAction action)
+void dtNavLinkBuilder::buildForAllEdges(rcContext& ctx, const dtLinkBuilderConfig& builderConfig, dtNavLinkAction action)
 {
-	for (int i = 0; i < m_nedges; ++i)
+	for (int i = 0; i < m_edges.Num(); ++i)
 	{
 		EdgeSampler sampler;
 		const bool success = sampleEdge(builderConfig, action, m_edges[i].sp, m_edges[i].sq, &sampler);
@@ -518,13 +465,15 @@ void dtNavLinkBuilder::buildForAllEdges(const dtLinkBuilderConfig& builderConfig
 			addEdgeLinks(builderConfig, &sampler);
 		}
 	}
-	
+
+	ctx.log(RC_LOG_PROGRESS, "   %i links added.", m_links.Num());
+
 	filterJumpOverLinks();
 }
 
 void dtNavLinkBuilder::debugBuildEdge(const dtLinkBuilderConfig& builderConfig, dtNavLinkAction action, int edgeIndex, EdgeSampler& sampler)
 {
-	if (edgeIndex >= m_nedges)
+	if (edgeIndex >= m_edges.Num())
 	{
 		return;
 	}
@@ -600,11 +549,9 @@ bool dtNavLinkBuilder::checkHeightfieldCollision(const dtReal x, const dtReal ym
 	
 	const int w = m_solid->width;
 	const int h = m_solid->height;
-	const dtReal cs = m_solid->cs;
-	const dtReal ch = m_solid->ch;
 	const rcReal* orig = m_solid->bmin;
-	const int ix = (int)floorf((x - orig[0])/cs);
-	const int iz = (int)floorf((z - orig[2])/cs);
+	const int ix = (int)dtFloor((x - orig[0])*m_invCs);
+	const int iz = (int)dtFloor((z - orig[2])*m_invCs);
 	
 	if (ix < 0 || iz < 0 || ix > w || iz > h)
 	{
@@ -619,8 +566,8 @@ bool dtNavLinkBuilder::checkHeightfieldCollision(const dtReal x, const dtReal ym
 	
 	while (s)
 	{
-		const float symin = orig[1] + s->data.smin*ch;
-		const float symax = orig[1] + s->data.smax*ch;
+		const float symin = orig[1] + s->data.smin*m_ch;
+		const float symax = orig[1] + s->data.smax*m_ch;
 		if (overlapRange(ymin, ymax, symin, symax))
 			return true;
 		
@@ -633,13 +580,15 @@ bool dtNavLinkBuilder::checkHeightfieldCollision(const dtReal x, const dtReal ym
 // Returns true if none of the samples ymin, ymax collide with the heghtfield.
 bool dtNavLinkBuilder::isTrajectoryClear(const dtReal* pa, const dtReal* pb, const Trajectory2D* tra) const
 {
-	for (int i = 0; i < tra->nsamples; ++i)
+	const int nsamples = tra->samples.Num();
+	const float invLastSample = 1.f / (nsamples-1);
+	for (int i = 0; i < nsamples; ++i)
 	{
 		dtReal p[3];
-		const TrajectorySample* s = &tra->samples[i];
-		const float u = (float)i / (float)(tra->nsamples-1);
+		const TrajectorySample& s = tra->samples[i];
+		const float u = (float)i * invLastSample;
 		dtVlerp(p, pa, pb, u);
-		if (checkHeightfieldCollision(p[0], p[1] + s->ymin, p[1] + s->ymax, p[2]))
+		if (checkHeightfieldCollision(p[0], p[1] + s.ymin, p[1] + s.ymax, p[2]))
 		{
 			return false;
 		}
@@ -718,29 +667,26 @@ void dtNavLinkBuilder::initTrajectory(Trajectory2D* tra) const
 	
 	const float dx = pb[0] - pa[0];
 	const int nsamples  = dtMax(2, (int)ceilf(dx*m_invCs));
-	tra->nsamples = nsamples;
-	// TODO: profile, preallocate, use dtChunkArray or else.
-	dtAssert(tra->samples == nullptr);
-	tra->samples = (TrajectorySample*)dtAlloc(sizeof(TrajectorySample)*nsamples, DT_ALLOC_PERM_TILE_LINK_BUILDER);
+	tra->samples.Reserve(nsamples);
 
 	const unsigned short lastSampleIndex = nsamples-1;
 	const float invLastIndex = 1.f/lastSampleIndex;
 	for (int i = 0; i < nsamples; ++i)
 	{
 		const float u = (float)i * invLastIndex;
-		TrajectorySample* s = &tra->samples[i];
-		s->x = dtLerp(pa[0], pb[0], u);
+		TrajectorySample& s = tra->samples.Emplace_GetRef();
+		s.x = dtLerp(pa[0], pb[0], u);
 
 		const float* spine = tra->spine;
 		unsigned char nspine = tra->nspine;
-		const float y0 = getHeight(s->x-m_linkBuilderConfig.agentRadius, spine, nspine);
-		const float y1 = getHeight(s->x+m_linkBuilderConfig.agentRadius, spine, nspine);
+		const float y0 = getHeight(s.x-m_linkBuilderConfig.agentRadius, spine, nspine);
+		const float y1 = getHeight(s.x+m_linkBuilderConfig.agentRadius, spine, nspine);
 
 		const float y = dtLerp(pa[1], pb[1], u);
 		
 		// Todo: revisit starting height and trajectory heights to handle ledge corner when jumping down.
-		s->ymin = dtMin(y0,y1) + (2*m_linkBuilderConfig.agentClimb) - y;	
-		s->ymax = dtMax(y0,y1) + m_linkBuilderConfig.agentHeight - y; 
+		s.ymin = dtMin(y0,y1) + (2*m_linkBuilderConfig.agentClimb) - y;	
+		s.ymax = dtMax(y0,y1) + m_linkBuilderConfig.agentHeight - y; 
 	}
 }
 
@@ -770,7 +716,7 @@ int dtNavLinkBuilder::findPotentialJumpOverEdges(const dtReal* sp, const dtReal*
 	PotentialSeg segs[MAX_SEGS];
 	int nsegs = 0;
 	
-	for (int i = 0; i < m_nedges; ++i)
+	for (int i = 0; i < m_edges.Num(); ++i)
 	{
 		dtReal p[3], lsp[3], lsq[3];
 		dtVsub(p, m_edges[i].sp, sp);
diff --git a/Engine/Source/Runtime/Navmesh/Private/Detour/DetourNavMeshBuilder.cpp b/Engine/Source/Runtime/Navmesh/Private/Detour/DetourNavMeshBuilder.cpp
--- a/Engine/Source/Runtime/Navmesh/Private/Detour/DetourNavMeshBuilder.cpp
+++ b/Engine/Source/Runtime/Navmesh/Private/Detour/DetourNavMeshBuilder.cpp
@@ -293,7 +293,6 @@ bool dtCreateNavMeshData(dtNavMeshCreateParams* params, unsigned char** outData,
 	// whose start point is inside the tile.
 	unsigned char* offMeshConClass = 0;
 	int storedOffMeshConCount = 0;
-	int offMeshConLinkCount = 0;
 	int storedOffMeshSegCount = 0;
 
 	if (params->offMeshConCount > 0)
@@ -343,6 +342,7 @@ bool dtCreateNavMeshData(dtNavMeshCreateParams* params, unsigned char** outData,
 
 		for (int i = 0; i < params->offMeshConCount; ++i)
 		{
+			
 			const dtOffMeshLinkCreateParams& offMeshCon = params->offMeshCons[i];
 			if (offMeshCon.type & DT_OFFMESH_CON_POINT)
 			{
@@ -359,12 +359,6 @@ bool dtCreateNavMeshData(dtNavMeshCreateParams* params, unsigned char** outData,
 					}
 				}
 
-				// Count how many links should be allocated for off-mesh connections.
-				if (offMeshConClass[i*2+0] == 0xff)
-					offMeshConLinkCount++;
-				if (offMeshConClass[i*2+1] == 0xff)
-					offMeshConLinkCount++;
-
 				if (offMeshConClass[i*2+0] == 0xff)
 					storedOffMeshConCount++;
 			}
diff --git a/Engine/Source/Runtime/Navmesh/Public/Detour/DetourNavLinkBuilder.h b/Engine/Source/Runtime/Navmesh/Public/Detour/DetourNavLinkBuilder.h
--- a/Engine/Source/Runtime/Navmesh/Public/Detour/DetourNavLinkBuilder.h
+++ b/Engine/Source/Runtime/Navmesh/Public/Detour/DetourNavLinkBuilder.h
@@ -51,11 +51,11 @@ class dtNavLinkBuilder
 	{
 		dtReal sp[3], sq[3];
 	};
-	Edge* m_edges = nullptr;
-	int m_nedges = 0;
+	TArray<Edge, TInlineAllocator<32>> m_edges;
 	
 public:	
 	static constexpr int MAX_SPINE = 8;
+	int getEdgeCount() const { return m_edges.Num(); }
 	
 	struct TrajectorySample
 	{
@@ -64,13 +64,11 @@ public:
 	
 	struct Trajectory2D
 	{
-		Trajectory2D() : samples(nullptr), nsamples(0), nspine(0) {}
-		NAVMESH_API ~Trajectory2D();
+		Trajectory2D() : nspine(0) {}
 
 		float spine[2*MAX_SPINE];		// [x,y] relative points representing the desired trajectory (2 spines)
-		TrajectorySample* samples;		// samples along trajectory slices to check for collision
+		TArray<TrajectorySample, TInlineAllocator<8>> samples;	// samples along trajectory slices to check for collision
 		
-		unsigned short nsamples;
 		unsigned char nspine;			// @todo: remove (use direclty MAX_SPINE) or make relative to trajectory type and config
 	};
 
@@ -103,7 +101,7 @@ private:
 		NAVMESH_API ~GroundSegment();
 
 		dtReal p[3], q[3];
-		TArray<GroundSample, TInlineAllocator<32>> gsamples;
+		TArray<GroundSample, TInlineAllocator<16>> gsamples;
 		unsigned short ngsamples;
 		unsigned short npass;
 	};
@@ -138,9 +136,7 @@ public:
 		JumpLinkFlag flags;
 		dtNavLinkAction action = DT_LINK_ACTION_UNSET;
 	};
-	JumpLink* m_links = nullptr;
-	int m_nlinks = 0;
-	int m_clinks = 0;
+	TArray<JumpLink, TInlineAllocator<16>> m_links;
 
 private:
 	int m_debugSelectedEdge = -1;
@@ -148,20 +144,17 @@ private:
 	friend NAVMESH_API void duDebugDrawNavLinkBuilder(struct duDebugDraw* dd, const dtNavLinkBuilder& linkBuilder, unsigned int drawFlags, const EdgeSampler* es);
 	
 public:
-	NAVMESH_API ~dtNavLinkBuilder();
-
 	// Loops through contours to store edge points in world coordinates.
 	NAVMESH_API bool findEdges(rcContext& ctx, const rcConfig& cfg, const dtLinkBuilderConfig& builderConfig,
 							   const struct dtTileCacheContourSet& lcset, const dtReal* orig,
 							   const rcHeightfield* solidHF, const rcCompactHeightfield* compactHF);
 
 	// For all edges, sample edges (sampleEdge) and add links to m_links
-	NAVMESH_API void buildForAllEdges(const dtLinkBuilderConfig& acfg, dtNavLinkAction action);
+	NAVMESH_API void buildForAllEdges(rcContext& ctx, const dtLinkBuilderConfig& acfg, dtNavLinkAction action);
 
 	NAVMESH_API void debugBuildEdge(const dtLinkBuilderConfig& acfg, dtNavLinkAction action, int edgeIndex, EdgeSampler& sampler);
 	
 private:
-	void cleanup();
 	void initTrajectory(Trajectory2D* tra) const;
 	bool isTrajectoryClear(const dtReal* pa, const dtReal* pb, const Trajectory2D* tra) const;
 	
@@ -184,10 +177,9 @@ private:
 	
 	void sampleAction(EdgeSampler* es) const;
 	
-	void filterJumpOverLinks() const;
+	void filterJumpOverLinks();
 
 	bool sampleEdge(const dtLinkBuilderConfig& builderConfig, dtNavLinkAction desiredAction, const dtReal* sp, const dtReal* sq, dtNavLinkBuilder::EdgeSampler* sampler) const;
-	JumpLink* addLink();
 	void addEdgeLinks(const dtLinkBuilderConfig& builderConfig, const EdgeSampler* es);
 };
 //@UE END
\ No newline at end of file
diff --git a/Engine/Source/Runtime/Navmesh/Public/Detour/DetourNavLinkBuilderConfig.h b/Engine/Source/Runtime/Navmesh/Public/Detour/DetourNavLinkBuilderConfig.h
--- a/Engine/Source/Runtime/Navmesh/Public/Detour/DetourNavLinkBuilderConfig.h
+++ b/Engine/Source/Runtime/Navmesh/Public/Detour/DetourNavLinkBuilderConfig.h
@@ -4,6 +4,7 @@
 /** Configuration for generated jump down links. */
 struct dtNavLinkBuilderJumpDownConfig
 {
+	bool enabled = true;
 	float jumpLength = 150.f; 
 	float jumpDistanceFromEdge = 10.f; 
 	float jumpMaxDepth = 150.f;
@@ -14,6 +15,7 @@ struct dtNavLinkBuilderJumpDownConfig
 /** Configuration for generated jump over links. */
 struct dtNavLinkBuilderJumpOverConfig
 {
+	bool enabled = true;
 	float jumpLength = 200.f; 
 	float jumpDistanceFromEdge = 100.f; 
 	float jumpHeight = 100.f; 
