Fix for potential crash in DetourNavMeshQuery storeWallSegment in the event that the poly->vertCount is 0.

#rnx

diff --git a/Engine/Source/Runtime/Navmesh/Private/Detour/DetourNavMeshQuery.cpp b/Engine/Source/Runtime/Navmesh/Private/Detour/DetourNavMeshQuery.cpp
--- a/Engine/Source/Runtime/Navmesh/Private/Detour/DetourNavMeshQuery.cpp
+++ b/Engine/Source/Runtime/Navmesh/Private/Detour/DetourNavMeshQuery.cpp
@@ -4418,7 +4418,13 @@ static void storeWallSegment(const dtNavMesh* nav, const dtMeshTile* tile, const
 	{
 		return;
 	}
-	
+
+	// Skip for polys that don't have vertices (e.g. nav links)
+	if (poly->vertCount == 0)
+	{
+		return;
+	}
+
 	const dtReal* va = &tile->verts[poly->verts[edge] * 3];
 	const dtReal* vb = &tile->verts[poly->verts[(edge + 1) % poly->vertCount] * 3];
 
