Navlink generation improvements
Ignore portal edges when looking for jumps
Added the filtering distance threshold to the link builder configuration
Improved link width computation
Improved filtering of overlapping links
Disable asymetric border sizes for now (using it can result in missing links)
Remove the extra 3 voxels on borders for link generation
Keep source edge on links for debug purposes
Add display of heightfield layers to the tile debug display
Added option to draw filtered links
Rebuild navmesh when navlink generation configuration changes

diff --git a/Engine/Source/Runtime/NavigationSystem/Private/NavMesh/LinkGenerationConfig.cpp b/Engine/Source/Runtime/NavigationSystem/Private/NavMesh/LinkGenerationConfig.cpp
--- a/Engine/Source/Runtime/NavigationSystem/Private/NavMesh/LinkGenerationConfig.cpp
+++ b/Engine/Source/Runtime/NavigationSystem/Private/NavMesh/LinkGenerationConfig.cpp
@@ -13,6 +13,7 @@ void FNavLinkGenerationJumpDownConfig::CopyToDetourConfig(dtNavLinkBuilderJumpDo
 	OutDetourConfig.jumpMaxDepth = JumpMaxDepth;
 	OutDetourConfig.jumpEndsHeightTolerance	= JumpEndsHeightTolerance;
 	OutDetourConfig.samplingSeparationFactor = SamplingSeparationFactor;
+	OutDetourConfig.filterDistanceThreshold = FilterDistanceThreshold;  
 }
 
 void FNavLinkGenerationJumpOverConfig::CopyToDetourConfig(dtNavLinkBuilderJumpOverConfig& OutDetourConfig) const
@@ -24,5 +25,6 @@ void FNavLinkGenerationJumpOverConfig::CopyToDetourConfig(dtNavLinkBuilderJumpOv
 	OutDetourConfig.jumpHeightTolerance = JumpHeightTolerance;
 	OutDetourConfig.jumpEndsHeightTolerance	= JumpEndsHeightTolerance;
 	OutDetourConfig.samplingSeparationFactor = SamplingSeparationFactor;
+	OutDetourConfig.filterDistanceThreshold = FilterDistanceThreshold;
 }
 #endif //WITH_RECAST
diff --git a/Engine/Source/Runtime/NavigationSystem/Private/NavMesh/RecastNavMesh.cpp b/Engine/Source/Runtime/NavigationSystem/Private/NavMesh/RecastNavMesh.cpp
--- a/Engine/Source/Runtime/NavigationSystem/Private/NavMesh/RecastNavMesh.cpp
+++ b/Engine/Source/Runtime/NavigationSystem/Private/NavMesh/RecastNavMesh.cpp
@@ -498,6 +498,7 @@ FRecastNavMeshTileGenerationDebug::FRecastNavMeshTileGenerationDebug()
 	bHeightfieldBounds = false;
 	bCompactHeightfield = false;
 	bCompactHeightfieldEroded = false;
+	bHeightFieldLayers = false;
 	bCompactHeightfieldRegions = false;
 	bCompactHeightfieldDistances = false;
 	bTileCacheLayerAreas = false;
@@ -3419,9 +3420,29 @@ void ARecastNavMesh::PostEditChangeChainProperty(FPropertyChangedChainEvent& Pro
 	
 	if (PropertyChangedChainEvent.Property != NULL)
 	{
-		const FName CategoryName = FObjectEditorUtils::GetCategoryFName(PropertyChangedChainEvent.Property);
+		FName CategoryName = FObjectEditorUtils::GetCategoryFName(PropertyChangedChainEvent.Property);
+
+		// If any, get the category of the parent node. 
+		FProperty* MemberProperty = nullptr;
+		const FEditPropertyChain::TDoubleLinkedListNode* PropertyNode = PropertyChangedChainEvent.PropertyChain.GetActiveNode();
+		if (PropertyNode)
+		{
+			const FEditPropertyChain::TDoubleLinkedListNode* PreviousNode = PropertyNode->GetPrevNode();
+			if (PreviousNode)
+			{
+				MemberProperty = PreviousNode->GetValue();
+				if (MemberProperty)
+				{
+					CategoryName = FObjectEditorUtils::GetCategoryFName(MemberProperty);
+				}
+			}
+		}
+		
 		if (CategoryName == NAME_Generation)
 		{
+			static const FName NAME_NavLinkJumpDownConfig = FName(TEXT("NavLinkJumpDownConfig"));
+			static const FName NAME_NavLinkJumpOverConfig = FName(TEXT("NavLinkJumpOverConfig"));
+			
 			const FName PropName = PropertyChangedChainEvent.Property->GetFName();
 			bool bRebuild = false;
 			
@@ -3467,6 +3488,10 @@ void ARecastNavMesh::PostEditChangeChainProperty(FPropertyChangedChainEvent& Pro
 
 				bRebuild = true;
 			}
+			else if (MemberProperty && (MemberProperty->GetFName() == NAME_NavLinkJumpDownConfig || MemberProperty->GetFName() == NAME_NavLinkJumpOverConfig))
+			{
+				bRebuild = true;
+			}
 
 			if (bRebuild)
 			{
diff --git a/Engine/Source/Runtime/NavigationSystem/Private/NavMesh/RecastNavMeshGenerator.cpp b/Engine/Source/Runtime/NavigationSystem/Private/NavMesh/RecastNavMeshGenerator.cpp
--- a/Engine/Source/Runtime/NavigationSystem/Private/NavMesh/RecastNavMeshGenerator.cpp
+++ b/Engine/Source/Runtime/NavigationSystem/Private/NavMesh/RecastNavMeshGenerator.cpp
@@ -84,7 +84,7 @@ namespace UE::NavMesh::Private
 	static bool bUseTightBoundExpansion = true;
 	static FAutoConsoleVariableRef CVarUseTightBoundExpansion(TEXT("ai.nav.UseTightBoundExpansion"), bUseTightBoundExpansion, TEXT("Active by default. Use an expansion of one AgentRadius. Set to false to revert to the previous behavior (2 AgentRadius)."), ECVF_Default);
 
-	static bool bUseAsymetricBorderSizes = true;
+	static bool bUseAsymetricBorderSizes = false;
 	static FAutoConsoleVariableRef CVarUseAsymetricBorderSizes(TEXT("ai.nav.UseAsymetricBorderSizes"), bUseAsymetricBorderSizes, TEXT("Active by default. When generating links, use asymetric tile border sizes to improve generation speed."), ECVF_Default);
 }
 
@@ -3557,6 +3557,23 @@ bool FRecastTileGenerator::GenerateNavigationData(FNavMeshBuildContext& BuildCon
 
 namespace UE::NavMesh::Private
 {
+	// Make sure LinkSpillDistance and CellSize has been computed before computing borders.
+	void ComputeConfigBorderSizes(const bool bGeneratingLinks, FRecastBuildConfig& InOutConfig)
+	{
+		int BorderForLinksVx = 0;
+		if (bGeneratingLinks)
+		{
+			BorderForLinksVx = (int)FMath::CeilToInt((rcReal)InOutConfig.LinkSpillDistance / InOutConfig.cs) + InOutConfig.walkableRadius;
+		}
+
+		// +1 for voxelization rounding, +1 for ledge neighbor access, +1 for occasional errors
+		const int BorderForAgentVx = InOutConfig.walkableRadius + 3;
+
+		// Borders must be at least the size of BorderForAgentVx.
+		InOutConfig.borderSize.low = UE::NavMesh::Private::bUseAsymetricBorderSizes ? BorderForAgentVx : FMath::Max(BorderForAgentVx, BorderForLinksVx);
+		InOutConfig.borderSize.high = FMath::Max(BorderForAgentVx, BorderForLinksVx);
+	}
+	
 #if RECAST_INTERNAL_DEBUG_DATA
 	void DrawHeightfield(const EHeightFieldRenderMode RenderMode, duDebugDraw* dd, const rcHeightfield& hf)
 	{
@@ -3661,6 +3678,11 @@ bool FRecastTileGenerator::GenerateCompressedLayers(FNavMeshBuildContext& BuildC
 #if RECAST_INTERNAL_DEBUG_DATA
 	if (IsTileDebugActive())
 	{
+		if (TileDebugSettings.bHeightFieldLayers)
+		{
+			duDebugDrawHeightfieldLayers(&BuildContext.InternalDebugData, *RasterContext.LayerSet);	
+		}
+		
 		if (TileDebugSettings.bCompactHeightfieldRegions)
 		{
 			duDebugDrawCompactHeightfieldRegions(&BuildContext.InternalDebugData, *CompactHF);	
@@ -3823,7 +3845,7 @@ dtStatus FRecastTileGenerator::BuildTileCacheLinks(FNavMeshBuildContext& BuildCo
 	// Make FNavigationLinks
 	for (const dtNavLinkBuilder::JumpLink& link : linkBuilder.m_links)
 	{
-		if (link.flags == dtNavLinkBuilder::INVALID)
+		if (link.flags == dtNavLinkBuilder::FILTERED)
 			continue;
 
 		// For now, just make a link using the center of the range.
@@ -4246,7 +4268,7 @@ bool FRecastTileGenerator::GenerateNavigationData(FNavMeshBuildContext& BuildCon
 	bool bGenDataLayer = true;
 	dtStatus status = DT_SUCCESS;
 
-	for (int32 LayerIdx = 0; LayerIdx < CompressedLayers.Num(); LayerIdx++)
+	for (int32 LayerIdx = 0; LayerIdx < CompressedLayers.Num() && LayerIdx < DirtyLayers.Num(); LayerIdx++)
 	{
 		if (DirtyLayers[LayerIdx] == false || !CompressedLayers[LayerIdx].IsValid())
 		{
@@ -4883,16 +4905,7 @@ void FRecastNavMeshGenerator::SetupTileConfig(const ENavigationDataResolution Ti
 	OutConfig.walkableRadius = FMath::CeilToInt(DestNavMesh->AgentRadius / CellSize);
 	OutConfig.maxStepFromWalkableSlope = OutConfig.cs * FMath::Tan(FMath::DegreesToRadians(OutConfig.walkableSlopeAngle));
 
-	int MaxLinkDistanceVx = 0;
-	if (DestNavMesh->bGenerateNavLinks)
-	{
-		MaxLinkDistanceVx = FMath::CeilToInt(OutConfig.LinkSpillDistance / CellSize);
-	}
-
-	// +1 for voxelization rounding, +1 for ledge neighbor access, +1 for occasional errors
-	const int BorderForAgentVx = OutConfig.walkableRadius + 3;
-	OutConfig.borderSize.low = UE::NavMesh::Private::bUseAsymetricBorderSizes ? BorderForAgentVx : BorderForAgentVx + MaxLinkDistanceVx;
-	OutConfig.borderSize.high = BorderForAgentVx + MaxLinkDistanceVx;
+	UE::NavMesh::Private::ComputeConfigBorderSizes(DestNavMesh->bGenerateNavLinks, OutConfig);
 
 	OutConfig.maxEdgeLen = (int32)(1200.0f / CellSize);
 
@@ -4962,7 +4975,6 @@ void FRecastNavMeshGenerator::ConfigureBuildProperties(FRecastBuildConfig& OutCo
 		}
 	}
 
-	int MaxLinkDistanceVx = 0;
 	if (DestNavMesh->bGenerateNavLinks)
 	{
 		// NavLink builder configuration
@@ -4974,7 +4986,6 @@ void FRecastNavMeshGenerator::ConfigureBuildProperties(FRecastBuildConfig& OutCo
 		const float JumpDownSpillDistance = JumpDown.bEnabled ? JumpDown.JumpLength - JumpDown.JumpDistanceFromEdge : 0.f;
 		const float JumpOverSpillDistance = JumpOver.bEnabled ? JumpOver.JumpLength - JumpOver.JumpDistanceFromEdge : 0.f;
 		OutConfig.LinkSpillDistance = FMath::Max(JumpDownSpillDistance, JumpOverSpillDistance);
-		MaxLinkDistanceVx = FMath::CeilToInt(OutConfig.LinkSpillDistance / CellSize);
 	}
 	
 	// store original sizes
@@ -4982,10 +4993,7 @@ void FRecastNavMeshGenerator::ConfigureBuildProperties(FRecastBuildConfig& OutCo
 	OutConfig.AgentMaxClimb = AgentMaxClimb;
 	OutConfig.AgentRadius = AgentRadius;
 
-	// +1 for voxelization rounding, +1 for ledge neighbor access, +1 for occasional errors
-	const int BorderForAgentVx = OutConfig.walkableRadius + 3;
-	OutConfig.borderSize.low = UE::NavMesh::Private::bUseAsymetricBorderSizes ? BorderForAgentVx : BorderForAgentVx + MaxLinkDistanceVx;
-	OutConfig.borderSize.high = BorderForAgentVx + MaxLinkDistanceVx;
+	UE::NavMesh::Private::ComputeConfigBorderSizes(DestNavMesh->bGenerateNavLinks, OutConfig);
 
 	OutConfig.maxEdgeLen = (int32)(1200.0f / CellSize);
 
diff --git a/Engine/Source/Runtime/NavigationSystem/Public/NavMesh/LinkGenerationConfig.h b/Engine/Source/Runtime/NavigationSystem/Public/NavMesh/LinkGenerationConfig.h
--- a/Engine/Source/Runtime/NavigationSystem/Public/NavMesh/LinkGenerationConfig.h
+++ b/Engine/Source/Runtime/NavigationSystem/Public/NavMesh/LinkGenerationConfig.h
@@ -18,16 +18,19 @@ struct FNavLinkGenerationJumpDownConfig
 	UPROPERTY(EditAnywhere, Config, Category = Settings)
 	bool bEnabled = true;
 	
-	// @todo: Rename, describe and find best defaults for those parameters.
+	/** Horizontal length of the jump. How far from the starting point we will look for ground. */ 
 	UPROPERTY(EditAnywhere, Config, Category = Settings, meta=(Units=cm, UIMin=0, ClampMin=0))
 	float JumpLength = 150.f; 
 
+	/** How far from the edge is the jump started */
 	UPROPERTY(EditAnywhere, Config, Category = Settings, meta=(Units=cm, UIMin=0, ClampMin=0))
 	float JumpDistanceFromEdge = 10.f; 
-	
+
+	/** How far below the starting height we want to look for landing ground. */
 	UPROPERTY(EditAnywhere, Config, Category = Settings, meta=(Units=cm, UIMin=0, ClampMin=0))
 	float JumpMaxDepth = 150.f;
 
+	/** Tolerance at both ends of the jump to find ground. */
 	UPROPERTY(EditAnywhere, Config, Category = Settings, meta=(Units=cm, UIMin=0, ClampMin=0))
 	float JumpEndsHeightTolerance = 80.f;
 
@@ -35,6 +38,10 @@ struct FNavLinkGenerationJumpDownConfig
     /*  Larger values improve generation speed but might introduce sampling errors.  */
 	UPROPERTY(EditAnywhere, Config, Category = Settings, meta=(UIMin=1, ClampMin=1))
 	float SamplingSeparationFactor = 1.f;
+	
+	/** When filtering similar links, distance used to compare between segment endpoints to match similar links. Use greater distance for more filtering (0 to deactivate filtering). */
+	UPROPERTY(EditAnywhere, Config, Category = Settings, meta=(Units=cm, UIMin=0, ClampMin=0))
+	float FilterDistanceThreshold = 80.f;
 
 #if WITH_RECAST	
 	/** Copy configuration to dtNavLinkBuilderJumpDownConfig. */
@@ -73,6 +80,10 @@ struct FNavLinkGenerationJumpOverConfig
 	UPROPERTY(EditAnywhere, Config, Category = Settings, meta=(UIMin=1, ClampMin=1))
 	float SamplingSeparationFactor = 1.f;
 
+	/** When filtering similar links, distance used to compare between segment endpoints to match similar links. Use greater distance for more filtering (0 to deactivate filtering). */
+	UPROPERTY(EditAnywhere, Config, Category = Settings, meta=(Units=cm, UIMin=0, ClampMin=0))
+	float FilterDistanceThreshold = 80.f;
+
 #if WITH_RECAST
 	/** Copy configuration to dtNavLinkBuilderJumpOverConfig. */
 	void CopyToDetourConfig(dtNavLinkBuilderJumpOverConfig& OutDetourConfig) const;
diff --git a/Engine/Source/Runtime/NavigationSystem/Public/NavMesh/LinkGenerationDebugFlags.h b/Engine/Source/Runtime/NavigationSystem/Public/NavMesh/LinkGenerationDebugFlags.h
--- a/Engine/Source/Runtime/NavigationSystem/Public/NavMesh/LinkGenerationDebugFlags.h
+++ b/Engine/Source/Runtime/NavigationSystem/Public/NavMesh/LinkGenerationDebugFlags.h
@@ -13,5 +13,6 @@ enum class ELinkGenerationDebugFlags : uint16
 	SelectedEdgeLandingSamples	= 1 << 4,
 	SelectedEdgeCollisions		= 1 << 5,
 	Links						= 1 << 6,
+	FilteredLinks				= 1 << 7,
 };
 ENUM_CLASS_FLAGS(ELinkGenerationDebugFlags);
diff --git a/Engine/Source/Runtime/NavigationSystem/Public/NavMesh/RecastNavMesh.h b/Engine/Source/Runtime/NavigationSystem/Public/NavMesh/RecastNavMesh.h
--- a/Engine/Source/Runtime/NavigationSystem/Public/NavMesh/RecastNavMesh.h
+++ b/Engine/Source/Runtime/NavigationSystem/Public/NavMesh/RecastNavMesh.h
@@ -499,6 +499,9 @@ struct FRecastNavMeshTileGenerationDebug
 	UPROPERTY(EditAnywhere, Category = Debug)
 	uint32 bCompactHeightfieldEroded : 1;
 
+	UPROPERTY(EditAnywhere, Category = Debug)
+	uint32 bHeightFieldLayers : 1;
+	
 	UPROPERTY(EditAnywhere, Category = Debug)
 	uint32 bCompactHeightfieldRegions : 1;
 
@@ -528,7 +531,7 @@ struct FRecastNavMeshTileGenerationDebug
 	uint16 LinkGenerationDebugFlags;
 
 	/** Using -1 as no selected edge. */
-	UPROPERTY(EditAnywhere, Category = Debug)
+	UPROPERTY(EditAnywhere, Category = Debug, meta=(UIMin=-1, ClampMin=-1))
 	int32 LinkGenerationSelectedEdge;
 };
 
diff --git a/Engine/Source/Runtime/Navmesh/Private/DebugUtils/DetourNavLinkDebugDraw.cpp b/Engine/Source/Runtime/Navmesh/Private/DebugUtils/DetourNavLinkDebugDraw.cpp
--- a/Engine/Source/Runtime/Navmesh/Private/DebugUtils/DetourNavLinkDebugDraw.cpp
+++ b/Engine/Source/Runtime/Navmesh/Private/DebugUtils/DetourNavLinkDebugDraw.cpp
@@ -252,12 +252,13 @@ void duDebugDrawNavLinkBuilder(duDebugDraw* dd, const dtNavLinkBuilder& linkBuil
 	
 	if (drawFlags & DRAW_LINKS)
 	{
+		const bool drawFilteredLinks = drawFlags & DRAW_FILTERED_LINKS;
 		unsigned int jumpDownCol0 = duLerpCol(duColor::blue, duColor::white, 200);
 		unsigned int jumpDownCol1 = duColor::blue;
 		unsigned int jumpOverCol0 = duLerpCol(duColor::lightGrey, duColor::white, 200);
 		unsigned int jumpOverCol1 = duColor::lightGrey;
 
-		auto selectColors = [&](dtNavLinkAction action, unsigned int& col0, unsigned int& col1)
+		auto selectColors = [&](dtNavLinkBuilder::JumpLinkFlag flag, dtNavLinkAction action, unsigned int& col0, unsigned int& col1)
 		{
 			if (action == DT_LINK_ACTION_JUMP_DOWN)
 			{
@@ -269,6 +270,12 @@ void duDebugDrawNavLinkBuilder(duDebugDraw* dd, const dtNavLinkBuilder& linkBuil
 				col0 = jumpOverCol0;
 				col1 = jumpOverCol1;
 			}
+
+			if (flag == dtNavLinkBuilder::FILTERED)
+			{
+				col0 = duColor::grey;
+				col1 = duColor::darkGrey;
+			}
 		};
 
 		unsigned int col0 = duColor::black;
@@ -280,17 +287,15 @@ void duDebugDrawNavLinkBuilder(duDebugDraw* dd, const dtNavLinkBuilder& linkBuil
 			dd->begin(DU_DRAW_QUADS);
 			for (const dtNavLinkBuilder::JumpLink& link : linkBuilder.m_links)
 			{
-				if (link.flags == dtNavLinkBuilder::INVALID)
+				if (!drawFilteredLinks && link.flags == dtNavLinkBuilder::FILTERED)
 					continue;
 
-				selectColors(link.action, col0, col1);
+				selectColors(link.flags, link.action, col0, col1);
 				
 				for (int j = 0; j < link.nspine-1; ++j)
 				{
 					int u = (j*255)/link.nspine;
 					unsigned int col = duTransCol(duLerpCol(col0,col1,u),128);
-					if (link.flags == dtNavLinkBuilder::INVALID)
-						col = duRGBA(255,0,0,64);
 					
 					dd->vertex(&link.spine1[j*3], col);
 					dd->vertex(&link.spine1[(j+1)*3], col);
@@ -303,10 +308,10 @@ void duDebugDrawNavLinkBuilder(duDebugDraw* dd, const dtNavLinkBuilder& linkBuil
 			dd->begin(DU_DRAW_LINES, 3.0f);
 			for (const dtNavLinkBuilder::JumpLink& link : linkBuilder.m_links)
 			{
-				if (link.flags == dtNavLinkBuilder::INVALID)
+				if (!drawFilteredLinks && link.flags == dtNavLinkBuilder::FILTERED)
 					continue;
 
-				selectColors(link.action, col0, col1);
+				selectColors(link.flags, link.action, col0, col1);
 				
 				for (int j = 0; j < link.nspine-1; ++j)
 				{
@@ -329,10 +334,10 @@ void duDebugDrawNavLinkBuilder(duDebugDraw* dd, const dtNavLinkBuilder& linkBuil
 			dd->begin(DU_DRAW_POINTS, 8.0f);
 			for (const dtNavLinkBuilder::JumpLink& link : linkBuilder.m_links)
 			{
-				if (link.flags == dtNavLinkBuilder::INVALID)
+				if (link.flags == dtNavLinkBuilder::FILTERED)
 					continue;
 
-				selectColors(link.action, col0, col1);
+				selectColors(link.flags, link.action, col0, col1);
 				
 				dd->vertex(&link.spine0[0], duDarkenCol(col1));
 				dd->vertex(&link.spine1[0], duDarkenCol(col1));
@@ -344,8 +349,9 @@ void duDebugDrawNavLinkBuilder(duDebugDraw* dd, const dtNavLinkBuilder& linkBuil
 			dd->begin(DU_DRAW_POINTS, 4.0f);
 			for (const dtNavLinkBuilder::JumpLink& link : linkBuilder.m_links)
 			{
-				if (link.flags == dtNavLinkBuilder::INVALID)
+				if (link.flags == dtNavLinkBuilder::FILTERED)
 					continue;
+				
 				dd->vertex(&link.spine0[0], duColor::lightGrey);
 				dd->vertex(&link.spine1[0], duColor::lightGrey);
 				dd->vertex(&link.spine0[(link.nspine-1)*3], duColor::lightGrey);
diff --git a/Engine/Source/Runtime/Navmesh/Private/Detour/DetourNavLinkBuilder.cpp b/Engine/Source/Runtime/Navmesh/Private/Detour/DetourNavLinkBuilder.cpp
--- a/Engine/Source/Runtime/Navmesh/Private/Detour/DetourNavLinkBuilder.cpp
+++ b/Engine/Source/Runtime/Navmesh/Private/Detour/DetourNavLinkBuilder.cpp
@@ -197,6 +197,20 @@ namespace UE::Detour::NavLink::Private
 		dst[1] = ax[1]*pt[0] + ay[1]*pt[1];
 		dst[2] = ax[2]*pt[0] + ay[2]*pt[1];
 	}
+
+	static float getDistanceThreshold(const dtLinkBuilderConfig& config, const dtNavLinkAction action)
+	{
+		switch(action)
+		{
+		case DT_LINK_ACTION_JUMP_DOWN:
+			return config.jumpDownConfig.filterDistanceThreshold;
+		case DT_LINK_ACTION_JUMP_OVER:
+			return config.jumpOverConfig.filterDistanceThreshold;
+		default:
+			dtAssert(false);
+			return 100.f;
+		}
+	}
 }
 
 dtNavLinkBuilder::GroundSegment::~GroundSegment()
@@ -247,6 +261,9 @@ bool dtNavLinkBuilder::findEdges(rcContext& ctx, const rcConfig& cfg, const dtLi
 			const unsigned short* va = &c.verts[k*4];
 			const unsigned short* vb = &c.verts[j*4];
 
+			if ((va[3] & 0xf) != 0xf)	// A direction is set, so it's a portal edge.
+				continue;
+
 			// Check k-j for matching contour
 			bool matchFound = false;
 			for (int ii = 0; ii < lcset.nconts; ++ii)
@@ -297,7 +314,7 @@ bool dtNavLinkBuilder::findEdges(rcContext& ctx, const rcConfig& cfg, const dtLi
 	return true;
 }
 
-void dtNavLinkBuilder::addEdgeLinks(const dtLinkBuilderConfig& builderConfig, const EdgeSampler* es)
+void dtNavLinkBuilder::addEdgeLinks(const dtLinkBuilderConfig& builderConfig, const EdgeSampler* es, const int edgeIndex)
 {
 	TRACE_CPUPROFILER_EVENT_SCOPE(dtNavLinkBuilder::addEdgeLinks);
 	
@@ -330,7 +347,8 @@ void dtNavLinkBuilder::addEdgeLinks(const dtLinkBuilderConfig& builderConfig, co
 		groundSampleFlags.Add(kernel[(nkernel+1)/2]);
 	}
 
-	const int agentRadiusVx = builderConfig.agentRadius / builderConfig.cellSize;
+	const dtReal edgeLength = dtVdist(es->rigp, es->rigq);
+	const dtReal distanceBetweenSamples = edgeLength / (es->start.ngsamples-1);
 	
 	// Build segments
 	int start = -1;
@@ -346,10 +364,11 @@ void dtNavLinkBuilder::addEdgeLinks(const dtLinkBuilderConfig& builderConfig, co
 		{
 			if (!valid)
 			{
-				if (i - start > agentRadiusVx)
+				const dtReal freeWidth = ((i-start)-1)*distanceBetweenSamples;
+				if (freeWidth >= builderConfig.agentRadius)
 				{
-					const float u0 = (float)start/(float)nsamples;
-					const float u1 = (float)i/(float)nsamples;
+					const float u0 = (float)start/(float)(nsamples-1);
+					const float u1 = (float)(i-1)/(float)(nsamples-1);
 
 					dtReal sp[3], sq[3], ep[3], eq[3];
 
@@ -364,6 +383,9 @@ void dtNavLinkBuilder::addEdgeLinks(const dtLinkBuilderConfig& builderConfig, co
 
 					JumpLink& link = m_links.Emplace_GetRef();
 
+#if !(UE_BUILD_SHIPPING || UE_BUILD_TEST)
+					link.debugSourceEdge = (short)edgeIndex;
+#endif
 					link.action = es->action;
 					link.flags = VALID;
 					link.nspine = es->trajectory.nspine;
@@ -404,19 +426,19 @@ void dtNavLinkBuilder::addEdgeLinks(const dtLinkBuilderConfig& builderConfig, co
 	}
 }
 
-void dtNavLinkBuilder::filterJumpOverLinks()
+void dtNavLinkBuilder::filterOverlappingLinks(const float edgeDistanceThreshold)
 {
-	TRACE_CPUPROFILER_EVENT_SCOPE(dtNavLinkBuilder::filterJumpOverLinks);
+	TRACE_CPUPROFILER_EVENT_SCOPE(dtNavLinkBuilder::filterOverlappingLinks);
 	
 	using namespace UE::Detour::NavLink::Private;
 	
 	// Filter out links which overlap
-	const float thresholdSquared = dtSqr(m_chf->cs*5.0f);
+	const float thresholdSquared = dtSqr(edgeDistanceThreshold);
 	
 	for (int i = 0; i < m_links.Num()-1; ++i)
 	{
 		JumpLink& li = m_links[i];
-		if (li.flags == INVALID)
+		if (li.flags == FILTERED)
 			continue;
 		
 		const dtReal* spi = &li.spine0[0];
@@ -427,26 +449,29 @@ void dtNavLinkBuilder::filterJumpOverLinks()
 		for (int j = i+1; j < m_links.Num(); ++j)
 		{
 			JumpLink& lj = m_links[j];
-			if (lj.flags == INVALID)
+			if (lj.flags == FILTERED)
 				continue;
 			
 			const dtReal* spj = &lj.spine0[0];
 			const dtReal* sqj = &lj.spine1[0];
 			const dtReal* epj = &lj.spine0[(lj.nspine-1)*3];
 			const dtReal* eqj = &lj.spine1[(lj.nspine-1)*3];
+
+			const dtReal d0 = dtDistancePtSegSqr(spj, epi, eqi);
+			const dtReal d1 = dtDistancePtSegSqr(sqj, epi, eqi);
+			const dtReal d2 = dtDistancePtSegSqr(epj, spi, sqi);
+			const dtReal d3 = dtDistancePtSegSqr(eqj, spi, sqi);
 			
-			const dtReal d0 = distSegSegSqr(spi,sqi, epj,eqj);
-			const dtReal d1 = distSegSegSqr(epi,eqi, spj,sqj);
-			
-			if (d0 < thresholdSquared && d1 < thresholdSquared)
+			if (d0 < thresholdSquared && d1 < thresholdSquared && d2 < thresholdSquared && d3 < thresholdSquared)
 			{
+				// Remove one of the link, keeping the wider one
 				if (dtVdistSqr(spi,sqi) > dtVdistSqr(spj,sqj))
 				{
-					lj.flags = INVALID;
+					lj.flags = FILTERED;
 				}
 				else
 				{
-					li.flags = INVALID;
+					li.flags = FILTERED;
 					break;
 				}
 			}
@@ -462,13 +487,14 @@ void dtNavLinkBuilder::buildForAllEdges(rcContext& ctx, const dtLinkBuilderConfi
 		const bool success = sampleEdge(builderConfig, action, m_edges[i].sp, m_edges[i].sq, &sampler);
 		if (success)
 		{
-			addEdgeLinks(builderConfig, &sampler);
+			addEdgeLinks(builderConfig, &sampler, i);
 		}
 	}
 
 	ctx.log(RC_LOG_PROGRESS, "   %i links added.", m_links.Num());
 
-	filterJumpOverLinks();
+	const float distanceThreshold = UE::Detour::NavLink::Private::getDistanceThreshold(builderConfig, action);
+	filterOverlappingLinks(distanceThreshold);
 }
 
 void dtNavLinkBuilder::debugBuildEdge(const dtLinkBuilderConfig& builderConfig, dtNavLinkAction action, int edgeIndex, EdgeSampler& sampler)
@@ -483,15 +509,15 @@ void dtNavLinkBuilder::debugBuildEdge(const dtLinkBuilderConfig& builderConfig,
 	const bool success = sampleEdge(builderConfig, action, m_edges[edgeIndex].sp, m_edges[edgeIndex].sq, &sampler);
 	if (success)
 	{
-		addEdgeLinks(builderConfig, &sampler);
+		addEdgeLinks(builderConfig, &sampler, edgeIndex);
 	}
 
-	filterJumpOverLinks();
+	const float distanceThreshold = UE::Detour::NavLink::Private::getDistanceThreshold(builderConfig, action);
+	filterOverlappingLinks(distanceThreshold);
 }
 
 bool dtNavLinkBuilder::getCompactHeightfieldHeight(const dtReal* pt, const dtReal hrange, dtReal* height) const
 {
-
 	const int chfWidth = m_chf->width;
 	const int chfHeight = m_chf->height;
 
@@ -863,13 +889,16 @@ int dtNavLinkBuilder::findPotentialJumpOverEdges(const dtReal* sp, const dtReal*
 	
 
 void dtNavLinkBuilder::initJumpDownRig(EdgeSampler* es, const dtReal* sp, const dtReal* sq,
-										const float jumpStartDist, const float jumpEndDist,
+										const float jumpStartDist, const float jumpLength,
 										const float jumpDownDist, const float groundRange)
 {
 	es->action = DT_LINK_ACTION_JUMP_DOWN;
+
+	// Set edge
 	dtVcopy(es->rigp, sp);
 	dtVcopy(es->rigq, sq);
-	
+
+	// Set axes
 	dtVsub(es->ax, sq, sp);
 	dtVnormalize(es->ax);
 	dtVset(es->az, es->ax[2], 0, -es->ax[0]);
@@ -882,7 +911,7 @@ void dtNavLinkBuilder::initJumpDownRig(EdgeSampler* es, const dtReal* sp, const
 	{
 		float* pt = &es->trajectory.spine[i*2];
 		const float u = (float)i/(float)(MAX_SPINE-1);
-		pt[0] = jumpStartDist + u * (jumpEndDist - jumpStartDist);
+		pt[0] = -jumpStartDist + (u*jumpLength);
 		pt[1] = u*u*u * jumpDownDist;
 	}
 
@@ -894,9 +923,12 @@ void dtNavLinkBuilder::initJumpOverRig(EdgeSampler* es, const dtReal* sp, const
 										const float jumpHeight, const float groundRange)
 {
 	es->action = DT_LINK_ACTION_JUMP_OVER;
+
+	// Set edge
 	dtVcopy(es->rigp, sp);
 	dtVcopy(es->rigq, sq);
-	
+
+	// Set axes	
 	dtVsub(es->ax, sq, sp);
 	dtVnormalize(es->ax);
 	dtVset(es->az, es->ax[2], 0, -es->ax[0]);
@@ -928,11 +960,11 @@ bool dtNavLinkBuilder::sampleEdge(const dtLinkBuilderConfig& builderConfig, dtNa
 	{
 		const dtNavLinkBuilderJumpDownConfig& config = builderConfig.jumpDownConfig;
 		const float jumpStartDist = config.jumpDistanceFromEdge;
-		const float jumpEndDist = config.jumpLength;
+		const float jumpLength = config.jumpLength;
 		const float jumpDownDist = config.jumpMaxDepth;
 		const float groundRange = config.jumpEndsHeightTolerance;
 		samplingSeparationFactor = config.samplingSeparationFactor;
-		initJumpDownRig(es, sp, sq, -jumpStartDist, jumpEndDist, -jumpDownDist, groundRange);
+		initJumpDownRig(es, sp, sq, jumpStartDist, jumpLength, -jumpDownDist, groundRange);
 	}
 	else if (desiredAction == DT_LINK_ACTION_JUMP_OVER)
 	{
@@ -964,7 +996,7 @@ bool dtNavLinkBuilder::sampleEdge(const dtLinkBuilderConfig& builderConfig, dtNa
 		const float jumpHeight = config.jumpHeight;
 		const float groundRange = config.jumpEndsHeightTolerance;
 		samplingSeparationFactor = config.samplingSeparationFactor;
-		initJumpOverRig(es, &segs[ibest*6+0], &segs[ibest*6+3], -jumpStartDist, jumpStartDist, jumpHeight, groundRange);
+		initJumpOverRig(es, &segs[ibest*6+0], &segs[ibest*6+3], jumpStartDist, jumpStartDist, jumpHeight, groundRange);
 	}
 	
 	initTrajectory(&es->trajectory);
diff --git a/Engine/Source/Runtime/Navmesh/Public/DebugUtils/DetourNavLinkDebugDraw.h b/Engine/Source/Runtime/Navmesh/Public/DebugUtils/DetourNavLinkDebugDraw.h
--- a/Engine/Source/Runtime/Navmesh/Public/DebugUtils/DetourNavLinkDebugDraw.h
+++ b/Engine/Source/Runtime/Navmesh/Public/DebugUtils/DetourNavLinkDebugDraw.h
@@ -17,6 +17,7 @@ enum duNavLinkBuilderDrawFlags
 	DRAW_LAND_SAMPLES =		1 << 4,
 	DRAW_COLLISION_SLICES =	1 << 5,
 	DRAW_LINKS =			1 << 6,
+	DRAW_FILTERED_LINKS =	1 << 7,
 };
 
 NAVMESH_API void duDebugDrawNavLinkBuilder(duDebugDraw* dd, const dtNavLinkBuilder& linkBuilder, unsigned int drawFlags, const struct dtNavLinkBuilder::EdgeSampler* es);
diff --git a/Engine/Source/Runtime/Navmesh/Public/Detour/DetourNavLinkBuilder.h b/Engine/Source/Runtime/Navmesh/Public/Detour/DetourNavLinkBuilder.h
--- a/Engine/Source/Runtime/Navmesh/Public/Detour/DetourNavLinkBuilder.h
+++ b/Engine/Source/Runtime/Navmesh/Public/Detour/DetourNavLinkBuilder.h
@@ -124,7 +124,7 @@ public:
 
 	enum JumpLinkFlag : unsigned char
 	{
-		INVALID = 0,
+		FILTERED = 0,
 		VALID = 1,
 	};
 	
@@ -135,6 +135,9 @@ public:
 		int nspine;
 		JumpLinkFlag flags;
 		dtNavLinkAction action = DT_LINK_ACTION_UNSET;
+#if !(UE_BUILD_SHIPPING || UE_BUILD_TEST)		
+		short debugSourceEdge = -1;
+#endif
 	};
 	TArray<JumpLink, TInlineAllocator<16>> m_links;
 
@@ -177,9 +180,9 @@ private:
 	
 	void sampleAction(EdgeSampler* es) const;
 	
-	void filterJumpOverLinks();
+	void filterOverlappingLinks(const float edgeDistanceThreshold);
 
 	bool sampleEdge(const dtLinkBuilderConfig& builderConfig, dtNavLinkAction desiredAction, const dtReal* sp, const dtReal* sq, dtNavLinkBuilder::EdgeSampler* sampler) const;
-	void addEdgeLinks(const dtLinkBuilderConfig& builderConfig, const EdgeSampler* es);
+	void addEdgeLinks(const dtLinkBuilderConfig& builderConfig, const EdgeSampler* es, const int edgeIndex);
 };
 //@UE END
\ No newline at end of file
diff --git a/Engine/Source/Runtime/Navmesh/Public/Detour/DetourNavLinkBuilderConfig.h b/Engine/Source/Runtime/Navmesh/Public/Detour/DetourNavLinkBuilderConfig.h
--- a/Engine/Source/Runtime/Navmesh/Public/Detour/DetourNavLinkBuilderConfig.h
+++ b/Engine/Source/Runtime/Navmesh/Public/Detour/DetourNavLinkBuilderConfig.h
@@ -9,7 +9,8 @@ struct dtNavLinkBuilderJumpDownConfig
 	float jumpDistanceFromEdge = 10.f; 
 	float jumpMaxDepth = 150.f;
 	float jumpEndsHeightTolerance = 50.f;
-	float samplingSeparationFactor = 1.f;	
+	float samplingSeparationFactor = 1.f;
+	float filterDistanceThreshold = 80.f;
 };
 
 /** Configuration for generated jump over links. */
@@ -22,4 +23,5 @@ struct dtNavLinkBuilderJumpOverConfig
 	float jumpHeightTolerance = 100.f; 
 	float jumpEndsHeightTolerance = 50.f;
 	float samplingSeparationFactor = 1.f;
+	float filterDistanceThreshold = 80.f;
 };
