[Navmesh] Recast AStart log name fix

diff --git a/Engine/Source/Runtime/Navmesh/Private/Detour/DetourNavMeshQuery.cpp b/Engine/Source/Runtime/Navmesh/Private/Detour/DetourNavMeshQuery.cpp
--- a/Engine/Source/Runtime/Navmesh/Private/Detour/DetourNavMeshQuery.cpp
+++ b/Engine/Source/Runtime/Navmesh/Private/Detour/DetourNavMeshQuery.cpp
@@ -28,8 +28,8 @@ DEFINE_LOG_CATEGORY_STATIC(LogDebugRaycastCrash, All, All);
 // To debug the recast AStar enable this define
 #define ENABLE_RECAST_ASTAR_LOGGING 0
 #if ENABLE_RECAST_ASTAR_LOGGING
-	DEFINE_LOG_CATEGORY_STATIC(LogRaycastAStart, Warning, All);
-	#define UE_RECAST_ASTAR_LOG(Verbosity, Format, ...) UE_LOG(LogRaycastAStart, Verbosity, Format, __VA_ARGS__)
+	DEFINE_LOG_CATEGORY_STATIC(LogRecastAStar, Display, All);
+	#define UE_RECAST_ASTAR_LOG(Verbosity, Format, ...) UE_LOG(LogRecastAStar, Verbosity, Format, __VA_ARGS__)
 #else
 	#define UE_RECAST_ASTAR_LOG(...)
 #endif
@@ -1674,7 +1674,7 @@ dtStatus dtNavMeshQuery::findPath(dtPolyRef startRef, dtPolyRef endRef,
 				|| !filter->isValidLinkSide(link.side))
 				//@UE END
 			{
-				UE_RECAST_ASTAR_LOG( Warning, TEXT("Filtered %lld from %lld"), neighbourRef, bestRef);
+				UE_RECAST_ASTAR_LOG(Display, TEXT("Filtered %lld from %lld"), neighbourRef, bestRef);
 				continue;
 			}
 			
@@ -1686,21 +1686,21 @@ dtStatus dtNavMeshQuery::findPath(dtPolyRef startRef, dtPolyRef endRef,
 			
 			if (!filter->passFilter(neighbourRef, neighbourTile, neighbourPoly) || !passLinkFilterByRef(neighbourTile, neighbourRef))
 			{
-				UE_RECAST_ASTAR_LOG(Warning, TEXT("Filtered %lld from %lld"), neighbourRef, bestRef);
+				UE_RECAST_ASTAR_LOG(Display, TEXT("Filtered %lld from %lld"), neighbourRef, bestRef);
 				continue;
 			}
 
 			dtNode* neighbourNode = m_nodePool->getNode(neighbourRef);
 			if (!neighbourNode)
 			{
-				UE_RECAST_ASTAR_LOG(Warning, TEXT("Reach Limit %lld from %lld"), neighbourRef, bestRef);
+				UE_RECAST_ASTAR_LOG(Display, TEXT("Reach Limit %lld from %lld"), neighbourRef, bestRef);
 				status |= DT_OUT_OF_NODES;
 				continue;
 			}
 //@UE BEGIN
 			else if (shouldIgnoreClosedNodes && (neighbourNode->flags & DT_NODE_CLOSED) != 0)
 			{
-				UE_RECAST_ASTAR_LOG(Warning, TEXT("Skipping closed %lld from %lld"), neighbourRef, bestRef);
+				UE_RECAST_ASTAR_LOG(Display, TEXT("Skipping closed %lld from %lld"), neighbourRef, bestRef);
 				continue;
 			}
 //@UE END
@@ -1745,27 +1745,27 @@ dtStatus dtNavMeshQuery::findPath(dtPolyRef startRef, dtPolyRef endRef,
 			// The node is already in open list and the new result is worse, skip.
 			if ((neighbourNode->flags & DT_NODE_OPEN) && total >= neighbourNode->total)
 			{
-				UE_RECAST_ASTAR_LOG(Warning, TEXT("Skipping new cost higher %lld from %lld cost %f total %f prev cost %f"), neighbourRef, bestRef, cost, total, neighbourNode->total);
+				UE_RECAST_ASTAR_LOG(Display, TEXT("Skipping new cost higher %lld from %lld cost %f total %f prev cost %f"), neighbourRef, bestRef, cost, total, neighbourNode->total);
 				continue;
 			}
 
 			// The node is already visited and process, and the new result is worse, skip.
 			if ((neighbourNode->flags & DT_NODE_CLOSED) && total >= neighbourNode->total)
 			{
-				UE_RECAST_ASTAR_LOG(Warning, TEXT("Skipping new cost higher %lld from %lld cost %f total %f prev cost %f"), neighbourRef, bestRef, cost, total, neighbourNode->total);
+				UE_RECAST_ASTAR_LOG(Display, TEXT("Skipping new cost higher %lld from %lld cost %f total %f prev cost %f"), neighbourRef, bestRef, cost, total, neighbourNode->total);
 				continue;
 			}
 
 			// Cost of current link is DT_UNWALKABLE_POLY_COST, skip.
 			if (curCost == DT_UNWALKABLE_POLY_COST)
 			{
-				UE_RECAST_ASTAR_LOG(Warning, TEXT("Skipping unwalkable poly %lld from %lld cost %f total %f prev cost %f"), neighbourRef, bestRef, cost, total, neighbourNode->total);
+				UE_RECAST_ASTAR_LOG(Display, TEXT("Skipping unwalkable poly %lld from %lld cost %f total %f prev cost %f"), neighbourRef, bestRef, cost, total, neighbourNode->total);
 				continue;
 			}
 
 			if (total > costLimit) //@UE
 			{
-				UE_RECAST_ASTAR_LOG(Warning, TEXT("Skipping reach cost limit poly %lld from %lld cost %f total %f prev cost %f limit %f"), neighbourRef, bestRef, cost, total, neighbourNode->total, costLimit);
+				UE_RECAST_ASTAR_LOG(Display, TEXT("Skipping reach cost limit poly %lld from %lld cost %f total %f prev cost %f limit %f"), neighbourRef, bestRef, cost, total, neighbourNode->total, costLimit);
 				continue;
 			}
 
@@ -1781,7 +1781,7 @@ dtStatus dtNavMeshQuery::findPath(dtPolyRef startRef, dtPolyRef endRef,
 			{
 				// Already in open, update node location.
 				m_openList->modify(neighbourNode);
-				UE_RECAST_ASTAR_LOG(Warning, TEXT("Modifying %lld from %lld cost %f total %f"), neighbourRef, bestRef, cost, total);
+				UE_RECAST_ASTAR_LOG(Display, TEXT("Modifying %lld from %lld cost %f total %f"), neighbourRef, bestRef, cost, total);
 			}
 			else
 			{
@@ -1789,13 +1789,13 @@ dtStatus dtNavMeshQuery::findPath(dtPolyRef startRef, dtPolyRef endRef,
 				neighbourNode->flags |= DT_NODE_OPEN;
 				m_openList->push(neighbourNode);
 				m_queryNodes++;
-				UE_RECAST_ASTAR_LOG(Warning, TEXT("Pushing %lld from %lld cost %f total %f"), neighbourRef, bestRef, cost, total);
+				UE_RECAST_ASTAR_LOG(Display, TEXT("Pushing %lld from %lld cost %f total %f"), neighbourRef, bestRef, cost, total);
 			}
 			
 			// Update nearest node to target so far.
 			if (heuristic < lastBestNodeCost)
 			{
-				UE_RECAST_ASTAR_LOG(Warning, TEXT("New best path %lld from %lld new best heuristic %f prev best heuristic %f"), neighbourRef, bestRef, heuristic, lastBestNodeCost);
+				UE_RECAST_ASTAR_LOG(Display, TEXT("New best path %lld from %lld new best heuristic %f prev best heuristic %f"), neighbourRef, bestRef, heuristic, lastBestNodeCost);
 				lastBestNodeCost = heuristic;
 				lastBestNode = neighbourNode;
 			}
