Navlink generation improvements
Trim source edge tips to account for edges overlapping rasterization borders
Fix some links not connecting to navmesh because they where generated to high
Fix to the 'jump over starting distance'
Addition of descriptions to jump config properties
Small improvements to ARecastNavMesh::PostEditChangeChainProperty

diff --git a/Engine/Source/Runtime/NavigationSystem/Private/NavMesh/LinkGenerationConfig.cpp b/Engine/Source/Runtime/NavigationSystem/Private/NavMesh/LinkGenerationConfig.cpp
--- a/Engine/Source/Runtime/NavigationSystem/Private/NavMesh/LinkGenerationConfig.cpp
+++ b/Engine/Source/Runtime/NavigationSystem/Private/NavMesh/LinkGenerationConfig.cpp
@@ -19,10 +19,10 @@ void FNavLinkGenerationJumpDownConfig::CopyToDetourConfig(dtNavLinkBuilderJumpDo
 void FNavLinkGenerationJumpOverConfig::CopyToDetourConfig(dtNavLinkBuilderJumpOverConfig& OutDetourConfig) const
 {
 	OutDetourConfig.enabled = bEnabled;
-	OutDetourConfig.jumpLength = JumpLength;
-	OutDetourConfig.jumpDistanceFromEdge = JumpDistanceFromEdge;
+	OutDetourConfig.jumpGapWidth = JumpGapWidth;
+	OutDetourConfig.jumpGapHeightTolerance = JumpGapHeightTolerance;
+	OutDetourConfig.jumpDistanceFromGapCenter = JumpDistanceFromGapCenter;
 	OutDetourConfig.jumpHeight = JumpHeight;
-	OutDetourConfig.jumpHeightTolerance = JumpHeightTolerance;
 	OutDetourConfig.jumpEndsHeightTolerance	= JumpEndsHeightTolerance;
 	OutDetourConfig.samplingSeparationFactor = SamplingSeparationFactor;
 	OutDetourConfig.filterDistanceThreshold = FilterDistanceThreshold;
diff --git a/Engine/Source/Runtime/NavigationSystem/Private/NavMesh/RecastNavMesh.cpp b/Engine/Source/Runtime/NavigationSystem/Private/NavMesh/RecastNavMesh.cpp
--- a/Engine/Source/Runtime/NavigationSystem/Private/NavMesh/RecastNavMesh.cpp
+++ b/Engine/Source/Runtime/NavigationSystem/Private/NavMesh/RecastNavMesh.cpp
@@ -3458,24 +3458,21 @@ void ARecastNavMesh::PostEditChangeChainProperty(FPropertyChangedChainEvent& Pro
 
 		// If any, get the category of the parent node. 
 		FProperty* MemberProperty = nullptr;
-		const FEditPropertyChain::TDoubleLinkedListNode* PropertyNode = PropertyChangedChainEvent.PropertyChain.GetActiveNode();
+
+		const FEditPropertyChain::TDoubleLinkedListNode* PropertyNode = PropertyChangedChainEvent.PropertyChain.GetActiveMemberNode();
 		if (PropertyNode)
 		{
-			const FEditPropertyChain::TDoubleLinkedListNode* PreviousNode = PropertyNode->GetPrevNode();
-			if (PreviousNode)
+			MemberProperty = PropertyNode->GetValue();
+			if (MemberProperty)
 			{
-				MemberProperty = PreviousNode->GetValue();
-				if (MemberProperty)
-				{
-					CategoryName = FObjectEditorUtils::GetCategoryFName(MemberProperty);
-				}
+				CategoryName = FObjectEditorUtils::GetCategoryFName(MemberProperty);
 			}
 		}
 		
 		if (CategoryName == NAME_Generation)
 		{
-			static const FName NAME_NavLinkJumpDownConfig = FName(TEXT("NavLinkJumpDownConfig"));
-			static const FName NAME_NavLinkJumpOverConfig = FName(TEXT("NavLinkJumpOverConfig"));
+			static const FName NAME_NavLinkJumpDownConfig = GET_MEMBER_NAME_CHECKED(ARecastNavMesh, NavLinkJumpDownConfig);
+			static const FName NAME_NavLinkJumpOverConfig = GET_MEMBER_NAME_CHECKED(ARecastNavMesh, NavLinkJumpOverConfig);
 			
 			const FName PropName = PropertyChangedChainEvent.Property->GetFName();
 			bool bRebuild = false;
diff --git a/Engine/Source/Runtime/NavigationSystem/Private/NavMesh/RecastNavMeshGenerator.cpp b/Engine/Source/Runtime/NavigationSystem/Private/NavMesh/RecastNavMeshGenerator.cpp
--- a/Engine/Source/Runtime/NavigationSystem/Private/NavMesh/RecastNavMeshGenerator.cpp
+++ b/Engine/Source/Runtime/NavigationSystem/Private/NavMesh/RecastNavMeshGenerator.cpp
@@ -3853,6 +3853,9 @@ dtStatus FRecastTileGenerator::BuildTileCacheLinks(FNavMeshBuildContext& BuildCo
 		dtVlerp(midA, &link.spine0[0], &link.spine1[0], 0.5);
 		dtReal midB[3];
 		dtVlerp(midB, &link.spine0[(link.nspine-1)*3], &link.spine1[(link.nspine-1)*3], 0.5);
+		// Since trajectory validation starts at agentClimb height to ignore small bumps, remove the offset for the actual link height.
+		midA[1] -= linkBuilderConfig.agentClimb;
+		midB[1] -= linkBuilderConfig.agentClimb;
 		OutGeneratedLinks.Add(FNavigationLink(Recast2UnrealPoint(midA), Recast2UnrealPoint(midB)));
 	}
 	
@@ -4984,7 +4987,7 @@ void FRecastNavMeshGenerator::ConfigureBuildProperties(FRecastBuildConfig& OutCo
 		JumpOver.CopyToDetourConfig(OutConfig.JumpOverConfig);
 
 		const float JumpDownSpillDistance = JumpDown.bEnabled ? JumpDown.JumpLength - JumpDown.JumpDistanceFromEdge : 0.f;
-		const float JumpOverSpillDistance = JumpOver.bEnabled ? JumpOver.JumpLength - JumpOver.JumpDistanceFromEdge : 0.f;
+		const float JumpOverSpillDistance = JumpOver.bEnabled ? JumpOver.JumpDistanceFromGapCenter : 0.f;
 		OutConfig.LinkSpillDistance = FMath::Max(JumpDownSpillDistance, JumpOverSpillDistance);
 	}
 	
diff --git a/Engine/Source/Runtime/NavigationSystem/Public/NavMesh/LinkGenerationConfig.h b/Engine/Source/Runtime/NavigationSystem/Public/NavMesh/LinkGenerationConfig.h
--- a/Engine/Source/Runtime/NavigationSystem/Public/NavMesh/LinkGenerationConfig.h
+++ b/Engine/Source/Runtime/NavigationSystem/Public/NavMesh/LinkGenerationConfig.h
@@ -22,7 +22,7 @@ struct FNavLinkGenerationJumpDownConfig
 	UPROPERTY(EditAnywhere, Config, Category = Settings, meta=(Units=cm, UIMin=0, ClampMin=0))
 	float JumpLength = 150.f; 
 
-	/** How far from the edge is the jump started */
+	/** How far from the edge is the jump started. */
 	UPROPERTY(EditAnywhere, Config, Category = Settings, meta=(Units=cm, UIMin=0, ClampMin=0))
 	float JumpDistanceFromEdge = 10.f; 
 
@@ -59,19 +59,23 @@ struct FNavLinkGenerationJumpOverConfig
 	UPROPERTY(EditAnywhere, Config, Category = Settings)
 	bool bEnabled = true;
 	
-	// @todo: Rename, describe and find best defaults for those parameters.
+	/* Maximum jumpable gap size used when matching edges to jump over. */
 	UPROPERTY(EditAnywhere, Config, Category = Settings, meta=(Units=cm, UIMin=0, ClampMin=0))
-	float JumpLength = 200.f; 
+	float JumpGapWidth = 200.f; 
 
+	/** Vertical tolerance used when matching edges to jump over. */
 	UPROPERTY(EditAnywhere, Config, Category = Settings, meta=(Units=cm, UIMin=0, ClampMin=0))
-	float JumpDistanceFromEdge = 100.f; 
+	float JumpGapHeightTolerance = 100.f; 
 	
+	/** How far from the center of the gap is the jump started. */
 	UPROPERTY(EditAnywhere, Config, Category = Settings, meta=(Units=cm, UIMin=0, ClampMin=0))
-	float JumpHeight = 100.f; 
+	float JumpDistanceFromGapCenter = 100.f; 
 
+	/** Height at the top of the jump trajectory. */
 	UPROPERTY(EditAnywhere, Config, Category = Settings, meta=(Units=cm, UIMin=0, ClampMin=0))
-	float JumpHeightTolerance = 100.f; 
+	float JumpHeight = 100.f; 
 
+	/** Tolerance at both ends of the jump to find ground. */
 	UPROPERTY(EditAnywhere, Config, Category = Settings, meta=(Units=cm, UIMin=0, ClampMin=0))
 	float JumpEndsHeightTolerance = 80.f;
 
diff --git a/Engine/Source/Runtime/Navmesh/Private/Detour/DetourNavLinkBuilder.cpp b/Engine/Source/Runtime/Navmesh/Private/Detour/DetourNavLinkBuilder.cpp
--- a/Engine/Source/Runtime/Navmesh/Private/Detour/DetourNavLinkBuilder.cpp
+++ b/Engine/Source/Runtime/Navmesh/Private/Detour/DetourNavLinkBuilder.cpp
@@ -225,6 +225,7 @@ bool dtNavLinkBuilder::findEdges(rcContext& ctx, const rcConfig& cfg, const dtLi
 	m_linkBuilderConfig = builderConfig;
 
 	m_cs = cfg.cs;
+	m_csSquared = dtSqr(cfg.cs);
 	m_ch = cfg.ch;
 	m_invCs = 1.0/cfg.cs;
 	m_solid = solidHF;
@@ -710,8 +711,7 @@ void dtNavLinkBuilder::initTrajectory(Trajectory2D* tra) const
 
 		const float y = dtLerp(pa[1], pb[1], u);
 		
-		// Todo: revisit starting height and trajectory heights to handle ledge corner when jumping down.
-		s.ymin = dtMin(y0,y1) + (2*m_linkBuilderConfig.agentClimb) - y;	
+		s.ymin = dtMin(y0,y1) + m_linkBuilderConfig.agentClimb - y;	
 		s.ymax = dtMax(y0,y1) + m_linkBuilderConfig.agentHeight - y; 
 	}
 }
@@ -726,7 +726,7 @@ int dtNavLinkBuilder::findPotentialJumpOverEdges(const dtReal* sp, const dtReal*
 	
 	// Find potential edges to join to.
 	const float widthRange = sqrtf(dtVdistSqr(sp,sq));
-	const float amin[3] = {0, -heightRange*0.5f, 0};
+	const float amin[3] = { 0, -heightRange*0.5f, 0 };
 	const float amax[3] = { widthRange, heightRange*0.5f, depthRange};
 	
 	const dtReal thr = cosf((180.0 - 45.0)/180.0*RC_PI);
@@ -890,14 +890,10 @@ int dtNavLinkBuilder::findPotentialJumpOverEdges(const dtReal* sp, const dtReal*
 
 void dtNavLinkBuilder::initJumpDownRig(EdgeSampler* es, const dtReal* sp, const dtReal* sq,
 										const float jumpStartDist, const float jumpLength,
-										const float jumpDownDist, const float groundRange)
+										const float jumpDownDist, const float groundRange) const
 {
 	es->action = DT_LINK_ACTION_JUMP_DOWN;
 
-	// Set edge
-	dtVcopy(es->rigp, sp);
-	dtVcopy(es->rigq, sq);
-
 	// Set axes
 	dtVsub(es->ax, sq, sp);
 	dtVnormalize(es->ax);
@@ -905,11 +901,27 @@ void dtNavLinkBuilder::initJumpDownRig(EdgeSampler* es, const dtReal* sp, const
 	dtVnormalize(es->az);
 	dtVset(es->ay, 0, 1, 0);
 	
+	// Set edge
+	const dtReal edgeLengthSqr = dtVdistSqr(sp, sq);
+	if (edgeLengthSqr > m_csSquared)
+	{
+		// Trim tips by cellSize to account for edges overlapping the rasterization borders.
+		// This avoids getting the wrong height in getCompactHeightfieldHeight that need to lookup multiple cells.
+		dtVmad(es->rigp, sp, es->ax, m_cs);
+		dtVmad(es->rigq, sq, es->ax, -m_cs);
+	}
+	else
+	{
+		// If it's impossible because the edge is too short, just keep the original edge.
+		dtVcopy(es->rigp, sp);
+		dtVcopy(es->rigq, sq);
+	}
+	
 	// Build action sampling spine.
 	es->trajectory.nspine = MAX_SPINE;
 	for (int i = 0; i < MAX_SPINE; ++i)
 	{
-		float* pt = &es->trajectory.spine[i*2];
+		float* pt = &es->trajectory.spine[i*2];			// pt: [xy] (x is toward jump end, y is up)
 		const float u = (float)i/(float)(MAX_SPINE-1);
 		pt[0] = -jumpStartDist + (u*jumpLength);
 		pt[1] = u*u*u * jumpDownDist;
@@ -969,8 +981,8 @@ bool dtNavLinkBuilder::sampleEdge(const dtLinkBuilderConfig& builderConfig, dtNa
 	else if (desiredAction == DT_LINK_ACTION_JUMP_OVER)
 	{
 		const dtNavLinkBuilderJumpOverConfig& config = builderConfig.jumpOverConfig;
-		const float jumpDist = config.jumpLength;
-		const float heightRange = config.jumpHeightTolerance;
+		const float jumpDist = config.jumpGapWidth;
+		const float heightRange = config.jumpGapHeightTolerance;
 		static constexpr int NSEGS = 8;
 		dtReal segs[NSEGS*6];
 		int nsegs = findPotentialJumpOverEdges(sp, sq, jumpDist, heightRange, segs, NSEGS);
@@ -992,11 +1004,11 @@ bool dtNavLinkBuilder::sampleEdge(const dtLinkBuilderConfig& builderConfig, dtNa
 			return false;
 		}
 
-		const float jumpStartDist = config.jumpDistanceFromEdge; 
+		const float jumpStartDist = config.jumpDistanceFromGapCenter; 
 		const float jumpHeight = config.jumpHeight;
 		const float groundRange = config.jumpEndsHeightTolerance;
 		samplingSeparationFactor = config.samplingSeparationFactor;
-		initJumpOverRig(es, &segs[ibest*6+0], &segs[ibest*6+3], jumpStartDist, jumpStartDist, jumpHeight, groundRange);
+		initJumpOverRig(es, &segs[ibest*6+0], &segs[ibest*6+3], -jumpStartDist, jumpStartDist, jumpHeight, groundRange);
 	}
 	
 	initTrajectory(&es->trajectory);
diff --git a/Engine/Source/Runtime/Navmesh/Public/Detour/DetourNavLinkBuilder.h b/Engine/Source/Runtime/Navmesh/Public/Detour/DetourNavLinkBuilder.h
--- a/Engine/Source/Runtime/Navmesh/Public/Detour/DetourNavLinkBuilder.h
+++ b/Engine/Source/Runtime/Navmesh/Public/Detour/DetourNavLinkBuilder.h
@@ -42,6 +42,7 @@ class dtNavLinkBuilder
 	dtLinkBuilderConfig m_linkBuilderConfig;
 
 	dtReal m_cs = 0;
+	dtReal m_csSquared = 0;
 	dtReal m_ch = 0;
 	dtReal m_invCs = 0;
 	const rcHeightfield* m_solid = nullptr;
@@ -66,7 +67,8 @@ public:
 	{
 		Trajectory2D() : nspine(0) {}
 
-		float spine[2*MAX_SPINE];		// [x,y] relative points representing the desired trajectory (2 spines)
+		float spine[2*MAX_SPINE];		// [x,y] relative points representing the desired trajectory (one spine with MAX_SPINE points)
+
 		TArray<TrajectorySample, TInlineAllocator<8>> samples;	// samples along trajectory slices to check for collision
 		
 		unsigned char nspine;			// @todo: remove (use direclty MAX_SPINE) or make relative to trajectory type and config
@@ -132,8 +134,8 @@ public:
 	{
 		dtReal spine0[MAX_SPINE*3];
 		dtReal spine1[MAX_SPINE*3];
-		int nspine;
-		JumpLinkFlag flags;
+		int nspine = 0;
+		JumpLinkFlag flags = VALID;
 		dtNavLinkAction action = DT_LINK_ACTION_UNSET;
 #if !(UE_BUILD_SHIPPING || UE_BUILD_TEST)		
 		short debugSourceEdge = -1;
@@ -165,9 +167,9 @@ private:
 								   const float depthRange, const float heightRange,
 								   dtReal* outSegs, const int maxOutSegs) const;
 	
-	static void initJumpDownRig(EdgeSampler* es, const dtReal* sp, const dtReal* sq,
+	void initJumpDownRig(EdgeSampler* es, const dtReal* sp, const dtReal* sq,
 						 const float jumpStartDist, const float jumpEndDist,
-						 const float jumpDownDist, const float groundRange);
+						 const float jumpDownDist, const float groundRange) const;
 
 	static void initJumpOverRig(EdgeSampler* es, const dtReal* sp, const dtReal* sq,
 						 const float jumpStartDist, const float jumpEndDist,
diff --git a/Engine/Source/Runtime/Navmesh/Public/Detour/DetourNavLinkBuilderConfig.h b/Engine/Source/Runtime/Navmesh/Public/Detour/DetourNavLinkBuilderConfig.h
--- a/Engine/Source/Runtime/Navmesh/Public/Detour/DetourNavLinkBuilderConfig.h
+++ b/Engine/Source/Runtime/Navmesh/Public/Detour/DetourNavLinkBuilderConfig.h
@@ -4,24 +4,56 @@
 /** Configuration for generated jump down links. */
 struct dtNavLinkBuilderJumpDownConfig
 {
+	/// Should this config be used to generate links.
 	bool enabled = true;
-	float jumpLength = 150.f; 
-	float jumpDistanceFromEdge = 10.f; 
+
+	/// Horizontal length of the jump. How far from the starting point we will look for ground. [Limit: > 0] [Units: wu]
+	float jumpLength = 150.f;
+
+	/// How far from the edge is the jump started. [Limit: > 0] [Units: wu]
+	float jumpDistanceFromEdge = 10.f;
+
+	/// How far below the starting height we want to look for landing ground. [Limit: > 0] [Units: wu]
 	float jumpMaxDepth = 150.f;
+
+	/// Tolerance at both ends of the jump to find ground. [Limit: > 0] [Units: wu]
 	float jumpEndsHeightTolerance = 50.f;
+
+	/// Value multiplied by CellSize to find the distance between sampling trajectories. Default is 1.
+	/// Larger values improve generation speed but might introduce sampling errors. [Limit: >= 1]
 	float samplingSeparationFactor = 1.f;
+
+	/// When filtering similar links, distance used to compare between segment endpoints to match similar links.
+	/// Use greater distance for more filtering (0 to deactivate filtering). [Limit: > 0] [Units: wu]
 	float filterDistanceThreshold = 80.f;
 };
 
 /** Configuration for generated jump over links. */
 struct dtNavLinkBuilderJumpOverConfig
 {
+	/// Should this config be used to generate links.
 	bool enabled = true;
-	float jumpLength = 200.f; 
-	float jumpDistanceFromEdge = 100.f; 
+
+	/// Maximum jumpable gap size used when matching edges to jump over. [Limit: > 0] [Units: wu]
+	float jumpGapWidth = 200.f;
+
+	/// Vertical tolerance used when matching edges to jump over. [Limit: > 0] [Units: wu]
+	float jumpGapHeightTolerance = 100.f; 
+
+	/// How far from the center of the gap is the jump started. [Limit: > 0] [Units: wu]
+	float jumpDistanceFromGapCenter = 100.f; 
+
+	/// Height at the top of the jump trajectory. [Limit: > 0] [Units: wu]
 	float jumpHeight = 100.f; 
-	float jumpHeightTolerance = 100.f; 
+
+	/// Tolerance at both ends of the jump to find ground. [Limit: > 0] [Units: wu]
 	float jumpEndsHeightTolerance = 50.f;
+
+	/// Value multiplied by CellSize to find the distance between sampling trajectories. Default is 1.
+	/// Larger values improve generation speed but might introduce sampling errors. [Limit: >= 1]
 	float samplingSeparationFactor = 1.f;
+
+	/// When filtering similar links, distance used to compare between segment endpoints to match similar links.
+	/// Use greater distance for more filtering (0 to deactivate filtering). [Limit: > 0] [Units: wu]
 	float filterDistanceThreshold = 80.f;
 };
