Improved AStar log redability and added details when filtering nodes.

diff --git a/Engine/Source/Runtime/Navmesh/Private/Detour/DetourNavMeshQuery.cpp b/Engine/Source/Runtime/Navmesh/Private/Detour/DetourNavMeshQuery.cpp
--- a/Engine/Source/Runtime/Navmesh/Private/Detour/DetourNavMeshQuery.cpp
+++ b/Engine/Source/Runtime/Navmesh/Private/Detour/DetourNavMeshQuery.cpp
@@ -1612,6 +1612,7 @@ dtStatus dtNavMeshQuery::findPath(dtPolyRef startRef, dtPolyRef endRef,
 	startNode->total = dtVdist(startPos, endPos) * H_SCALE;
 	startNode->id = startRef;
 	startNode->flags = DT_NODE_OPEN;
+	UE_RECAST_ASTAR_LOG(Display, TEXT("Start by pushing %lld"), startRef);
 	m_openList->push(startNode);
 	m_queryNodes++;
 
@@ -1627,6 +1628,7 @@ dtStatus dtNavMeshQuery::findPath(dtPolyRef startRef, dtPolyRef endRef,
 	{
 		// Remove node from open list and put it in closed list.
 		dtNode* bestNode = m_openList->pop();
+		UE_RECAST_ASTAR_LOG(Display, TEXT("   Pop %lld"), bestNode->id);
 		bestNode->flags &= ~DT_NODE_OPEN;
 		bestNode->flags |= DT_NODE_CLOSED;
 		
@@ -1674,7 +1676,7 @@ dtStatus dtNavMeshQuery::findPath(dtPolyRef startRef, dtPolyRef endRef,
 				|| !filter->isValidLinkSide(link.side))
 				//@UE END
 			{
-				UE_RECAST_ASTAR_LOG(Display, TEXT("Filtered %lld from %lld"), neighbourRef, bestRef);
+				UE_RECAST_ASTAR_LOG(Display, TEXT("      Filtered %lld from %lld, isValidLinkSide %i"), neighbourRef, bestRef, filter->isValidLinkSide(link.side));
 				continue;
 			}
 			
@@ -1686,21 +1688,23 @@ dtStatus dtNavMeshQuery::findPath(dtPolyRef startRef, dtPolyRef endRef,
 			
 			if (!filter->passFilter(neighbourRef, neighbourTile, neighbourPoly) || !passLinkFilterByRef(neighbourTile, neighbourRef))
 			{
-				UE_RECAST_ASTAR_LOG(Display, TEXT("Filtered %lld from %lld"), neighbourRef, bestRef);
+				UE_RECAST_ASTAR_LOG(Display, TEXT("      Filtered %lld from %lld, filter->isVirtual %i, passFilter %i, passLinkFilterByRef %i"),
+					neighbourRef, bestRef, filter->getIsVirtual(),
+					filter->passFilter(neighbourRef, neighbourTile, neighbourPoly), passLinkFilterByRef(neighbourTile, neighbourRef));
 				continue;
 			}
 
 			dtNode* neighbourNode = m_nodePool->getNode(neighbourRef);
 			if (!neighbourNode)
 			{
-				UE_RECAST_ASTAR_LOG(Display, TEXT("Reach Limit %lld from %lld"), neighbourRef, bestRef);
+				UE_RECAST_ASTAR_LOG(Display, TEXT("      Reach Limit %lld from %lld"), neighbourRef, bestRef);
 				status |= DT_OUT_OF_NODES;
 				continue;
 			}
 //@UE BEGIN
 			else if (shouldIgnoreClosedNodes && (neighbourNode->flags & DT_NODE_CLOSED) != 0)
 			{
-				UE_RECAST_ASTAR_LOG(Display, TEXT("Skipping closed %lld from %lld"), neighbourRef, bestRef);
+				UE_RECAST_ASTAR_LOG(Display, TEXT("      Skipping closed %lld from %lld"), neighbourRef, bestRef);
 				continue;
 			}
 //@UE END
@@ -1745,27 +1749,27 @@ dtStatus dtNavMeshQuery::findPath(dtPolyRef startRef, dtPolyRef endRef,
 			// The node is already in open list and the new result is worse, skip.
 			if ((neighbourNode->flags & DT_NODE_OPEN) && total >= neighbourNode->total)
 			{
-				UE_RECAST_ASTAR_LOG(Display, TEXT("Skipping new cost higher %lld from %lld cost %f total %f prev cost %f"), neighbourRef, bestRef, cost, total, neighbourNode->total);
+				UE_RECAST_ASTAR_LOG(Display, TEXT("      Skipping new cost higher %lld from %lld cost %f total %f prev cost %f"), neighbourRef, bestRef, cost, total, neighbourNode->total);
 				continue;
 			}
 
 			// The node is already visited and process, and the new result is worse, skip.
 			if ((neighbourNode->flags & DT_NODE_CLOSED) && total >= neighbourNode->total)
 			{
-				UE_RECAST_ASTAR_LOG(Display, TEXT("Skipping new cost higher %lld from %lld cost %f total %f prev cost %f"), neighbourRef, bestRef, cost, total, neighbourNode->total);
+				UE_RECAST_ASTAR_LOG(Display, TEXT("      Skipping new cost higher %lld from %lld cost %f total %f prev cost %f"), neighbourRef, bestRef, cost, total, neighbourNode->total);
 				continue;
 			}
 
 			// Cost of current link is DT_UNWALKABLE_POLY_COST, skip.
 			if (curCost == DT_UNWALKABLE_POLY_COST)
 			{
-				UE_RECAST_ASTAR_LOG(Display, TEXT("Skipping unwalkable poly %lld from %lld cost %f total %f prev cost %f"), neighbourRef, bestRef, cost, total, neighbourNode->total);
+				UE_RECAST_ASTAR_LOG(Display, TEXT("      Skipping unwalkable poly %lld from %lld cost %f total %f prev cost %f"), neighbourRef, bestRef, cost, total, neighbourNode->total);
 				continue;
 			}
 
 			if (total > costLimit) //@UE
 			{
-				UE_RECAST_ASTAR_LOG(Display, TEXT("Skipping reach cost limit poly %lld from %lld cost %f total %f prev cost %f limit %f"), neighbourRef, bestRef, cost, total, neighbourNode->total, costLimit);
+				UE_RECAST_ASTAR_LOG(Display, TEXT("      Skipping reach cost limit poly %lld from %lld cost %f total %f prev cost %f limit %f"), neighbourRef, bestRef, cost, total, neighbourNode->total, costLimit);
 				continue;
 			}
 
@@ -1781,7 +1785,7 @@ dtStatus dtNavMeshQuery::findPath(dtPolyRef startRef, dtPolyRef endRef,
 			{
 				// Already in open, update node location.
 				m_openList->modify(neighbourNode);
-				UE_RECAST_ASTAR_LOG(Display, TEXT("Modifying %lld from %lld cost %f total %f"), neighbourRef, bestRef, cost, total);
+				UE_RECAST_ASTAR_LOG(Display, TEXT("   Modifying %lld from %lld cost %f total %f"), neighbourRef, bestRef, cost, total);
 			}
 			else
 			{
@@ -1789,7 +1793,7 @@ dtStatus dtNavMeshQuery::findPath(dtPolyRef startRef, dtPolyRef endRef,
 				neighbourNode->flags |= DT_NODE_OPEN;
 				m_openList->push(neighbourNode);
 				m_queryNodes++;
-				UE_RECAST_ASTAR_LOG(Display, TEXT("Pushing %lld from %lld cost %f total %f"), neighbourRef, bestRef, cost, total);
+				UE_RECAST_ASTAR_LOG(Display, TEXT("   Pushing %lld from %lld cost %f total %f"), neighbourRef, bestRef, cost, total);
 			}
 			
 			// Update nearest node to target so far.
diff --git a/Engine/Source/Runtime/Navmesh/Public/Detour/DetourNavMeshQuery.h b/Engine/Source/Runtime/Navmesh/Public/Detour/DetourNavMeshQuery.h
--- a/Engine/Source/Runtime/Navmesh/Public/Detour/DetourNavMeshQuery.h
+++ b/Engine/Source/Runtime/Navmesh/Public/Detour/DetourNavMeshQuery.h
@@ -127,6 +127,9 @@ protected:
 	}
 
 public:
+	/// Returns true if it's a virtual filter.
+	bool getIsVirtual() const { return isVirtual; }
+	
 	/// Returns true if the polygon can be visited.  (I.e. Is traversable.)
 	///  @param[in]		ref		The reference id of the polygon test.
 	///  @param[in]		tile	The tile containing the polygon.
