The navlink jump trajectory is now computed with a JumpHeight paramter.
The trajector will respect the JumpLength, JumpHeight and JumpMaxDepth configurations.
CL also includes removal of unused methods.

diff --git a/Engine/Source/Runtime/NavigationSystem/Private/NavMesh/LinkGenerationConfig.cpp b/Engine/Source/Runtime/NavigationSystem/Private/NavMesh/LinkGenerationConfig.cpp
--- a/Engine/Source/Runtime/NavigationSystem/Private/NavMesh/LinkGenerationConfig.cpp
+++ b/Engine/Source/Runtime/NavigationSystem/Private/NavMesh/LinkGenerationConfig.cpp
@@ -12,6 +12,7 @@ void FNavLinkGenerationJumpDownConfig::CopyToDetourConfig(dtNavLinkBuilderJumpDo
 	OutDetourConfig.jumpLength = JumpLength;
 	OutDetourConfig.jumpDistanceFromEdge = JumpDistanceFromEdge;
 	OutDetourConfig.jumpMaxDepth = JumpMaxDepth;
+	OutDetourConfig.jumpHeight = JumpHeight;
 	OutDetourConfig.jumpEndsHeightTolerance	= JumpEndsHeightTolerance;
 	OutDetourConfig.samplingSeparationFactor = SamplingSeparationFactor;
 	OutDetourConfig.filterDistanceThreshold = FilterDistanceThreshold;
diff --git a/Engine/Source/Runtime/NavigationSystem/Private/NavMesh/RecastNavMeshGenerator.cpp b/Engine/Source/Runtime/NavigationSystem/Private/NavMesh/RecastNavMeshGenerator.cpp
--- a/Engine/Source/Runtime/NavigationSystem/Private/NavMesh/RecastNavMeshGenerator.cpp
+++ b/Engine/Source/Runtime/NavigationSystem/Private/NavMesh/RecastNavMeshGenerator.cpp
@@ -3783,9 +3783,13 @@ dtStatus FRecastTileGenerator::BuildTileCacheLinks(FNavMeshBuildContext& BuildCo
 	}
 
 	const dtReal* orig = layer.header->bmin;
+
 	dtLinkBuilderConfig linkBuilderConfig;
 	linkBuilderConfig.jumpDownConfig = TileConfig.JumpDownConfig;
+	linkBuilderConfig.jumpDownConfig.init();
+	
 	linkBuilderConfig.jumpOverConfig = TileConfig.JumpOverConfig;
+	
 	linkBuilderConfig.agentRadius = TileConfig.walkableRadius * TileConfig.cs;
 	linkBuilderConfig.agentHeight = TileConfig.walkableHeight * TileConfig.ch;
 	linkBuilderConfig.agentClimb = TileConfig.walkableClimb * TileConfig.ch;
diff --git a/Engine/Source/Runtime/NavigationSystem/Public/NavMesh/LinkGenerationConfig.h b/Engine/Source/Runtime/NavigationSystem/Public/NavMesh/LinkGenerationConfig.h
--- a/Engine/Source/Runtime/NavigationSystem/Public/NavMesh/LinkGenerationConfig.h
+++ b/Engine/Source/Runtime/NavigationSystem/Public/NavMesh/LinkGenerationConfig.h
@@ -35,6 +35,10 @@ struct FNavLinkGenerationJumpDownConfig
 	UPROPERTY(EditAnywhere, Config, Category = Settings, meta=(Units=cm, UIMin=0, ClampMin=0))
 	float JumpMaxDepth = 150.f;
 
+	/** Peak height relative to the height of the starting point. */
+	UPROPERTY(EditAnywhere, Config, Category = Settings, meta=(Units=cm, UIMin=0, ClampMin=0))
+	float JumpHeight = 50.f;
+
 	/** Tolerance at both ends of the jump to find ground. */
 	UPROPERTY(EditAnywhere, Config, Category = Settings, meta=(Units=cm, UIMin=0, ClampMin=0))
 	float JumpEndsHeightTolerance = 80.f;
diff --git a/Engine/Source/Runtime/Navmesh/Private/Detour/DetourNavLinkBuilder.cpp b/Engine/Source/Runtime/Navmesh/Private/Detour/DetourNavLinkBuilder.cpp
--- a/Engine/Source/Runtime/Navmesh/Private/Detour/DetourNavLinkBuilder.cpp
+++ b/Engine/Source/Runtime/Navmesh/Private/Detour/DetourNavLinkBuilder.cpp
@@ -1,7 +1,6 @@
 // Copyright Epic Games, Inc. All Rights Reserved.
 #include "Detour/DetourNavLinkBuilder.h"
 
-#include "Detour/DetourAlloc.h"
 #include "Detour/DetourCommon.h"
 #include "DetourTileCache/DetourTileCacheBuilder.h"
 #include "Recast/Recast.h"
@@ -39,74 +38,6 @@ namespace UE::Detour::NavLink::Private
 		return t/d;
 	}
 
-	static void closestPtSegSeg(const dtReal* ap, const dtReal* aq, const dtReal* bp, const dtReal* bq, dtReal& s, dtReal& t)
-	{
-		constexpr dtReal EPSILON = 1e-6;
-		dtReal d1[3], d2[3], r[3];
-		dtVsub(d1, aq, ap);
-		dtVsub(d2, bq, bp);
-		dtVsub(r, ap, bp);
-		const dtReal a = dtVdot(d1,d1); // Squared length of segment S1, always nonnegative
-		const dtReal e = dtVdot(d2,d2); // Squared length of segment S2, always nonnegative
-		const dtReal f = dtVdot(d2,r);
-	
-		// Check if either or both segments degenerate into points
-		if (a <= EPSILON && e <= EPSILON)
-		{
-			// Both segments degenerate into points
-			s = t = 0.0;
-			return;
-		}
-	
-		if (a <= EPSILON)
-		{
-			s = 0.0;
-			t = f / e;
-			t = dtClamp(t, 0.0, 1.0);
-		}
-		else
-		{
-			const dtReal c = dtVdot(d1,r);
-			if (e <= EPSILON)
-			{
-				// Second segment degenerates into a point
-				t = 0.0;
-				s = dtClamp(-c / a, 0.0, 1.0); // t = 0 => s = (b*t - c) / a = -c / a
-			}
-			else
-			{
-				// The general nondegenerate case starts here
-				const dtReal b = dtVdot(d1,d2);
-				const dtReal denom = a*e-b*b; // Always nonnegative
-			
-				// If segments not parallel, compute closest point on L1 to L2, and
-				// clamp to segment S1. Else pick arbitrary s (here 0)
-				if (denom != 0.0)
-					s = dtClamp((b*f - c*e) / denom, 0.0, 1.0);
-				else
-					s = 0.0;
-			
-				// Compute point on L2 closest to S1(s) using
-				// t = Dot((p.start+D1*s)-q.start,D2) / Dot(D2,D2) = (b*s + f) / e
-				t = (b*s + f) / e;
-			
-				// If t in [0,1] done. Else clamp t, recompute s for the new value
-				// of t using s = Dot((q.start+D2*t)-p.start,D1) / Dot(D1,D1)= (t*b - c) / a
-				// and clamp s to [0, 1]
-				if (t < 0.0)
-				{
-					t = 0.0;
-					s = dtClamp(-c / a, 0.0, 1.0);
-				}
-				else if (t > 1.0)
-				{
-					t = 1.0;
-					s = dtClamp((b - c) / a, 0.0, 1.0);
-				}
-			}
-		}
-	}
-
 	static bool isectSegAABB(const dtReal* sp, const dtReal* sq,
 					  const float* amin, const float* amax,
 					  float& tmin, float& tmax)
@@ -153,16 +84,6 @@ namespace UE::Detour::NavLink::Private
 	
 		return true;
 	}
-
-	static dtReal distSegSegSqr(const dtReal* ap, const dtReal* aq, const dtReal* bp, const dtReal* bq)
-	{
-		dtReal s, t;
-		closestPtSegSeg(ap, aq, bp, bq, s,t);
-		dtReal anp[3], bnp[3];
-		dtVlerp(anp, ap,aq, s);
-		dtVlerp(bnp, bp,bq, t);
-		return dtVdistSqr(anp, bnp);
-	}
 	
 	static float getHeight(const float x, const float* pts, const int npts)
 	{
@@ -890,8 +811,7 @@ int dtNavLinkBuilder::findPotentialJumpOverEdges(const dtReal* sp, const dtReal*
 	
 
 void dtNavLinkBuilder::initJumpDownRig(EdgeSampler* es, const dtReal* sp, const dtReal* sq,
-										const float jumpStartDist, const float jumpLength,
-										const float jumpDownDist, const float groundRange) const
+									   const dtNavLinkBuilderJumpDownConfig& config) const
 {
 	es->action = DT_LINK_ACTION_JUMP_DOWN;
 
@@ -917,6 +837,16 @@ void dtNavLinkBuilder::initJumpDownRig(EdgeSampler* es, const dtReal* sp, const
 		dtVcopy(es->rigp, sp);
 		dtVcopy(es->rigq, sq);
 	}
+
+	// Parabolic equation y(x) = ax^2 + (-d/l - al)x
+	// Where 'a' is constant
+	//       'l' is the jump length from the starting point
+	//       'd' is the distance below the starting point
+
+	const float jumpStartDist = config.jumpDistanceFromEdge;
+	const float jumpLength = config.jumpLength;
+	const float a = config.cachedParabolaConstant;
+	const float downRatio = config.cachedDownRatio;		// -d/l
 	
 	// Build action sampling spine.
 	es->trajectory.nspine = MAX_SPINE;
@@ -925,10 +855,13 @@ void dtNavLinkBuilder::initJumpDownRig(EdgeSampler* es, const dtReal* sp, const
 		float* pt = &es->trajectory.spine[i*2];			// pt: [xy] (x is toward jump end, y is up)
 		const float u = (float)i/(float)(MAX_SPINE-1);
 		pt[0] = -jumpStartDist + (u*jumpLength);
-		pt[1] = u*u*u * jumpDownDist;
+
+		// Parabolic equation y(x) = ax^2 + (-d/l - al)x
+		//                    y(x) = x * (ax + (-d/l - al))
+		pt[1] = (u*jumpLength) * (a*(u*jumpLength) + (downRatio - a*jumpLength));
 	}
 
-	es->groundRange = groundRange;
+	es->groundRange = config.jumpEndsHeightTolerance;
 }
 
 void dtNavLinkBuilder::initJumpOverRig(EdgeSampler* es, const dtReal* sp, const dtReal* sq,
@@ -972,12 +905,8 @@ bool dtNavLinkBuilder::sampleEdge(const dtLinkBuilderConfig& builderConfig, dtNa
 	if (desiredAction == DT_LINK_ACTION_JUMP_DOWN)
 	{
 		const dtNavLinkBuilderJumpDownConfig& config = builderConfig.jumpDownConfig;
-		const float jumpStartDist = config.jumpDistanceFromEdge;
-		const float jumpLength = config.jumpLength;
-		const float jumpDownDist = config.jumpMaxDepth;
-		const float groundRange = config.jumpEndsHeightTolerance;
 		samplingSeparationFactor = config.samplingSeparationFactor;
-		initJumpDownRig(es, sp, sq, jumpStartDist, jumpLength, -jumpDownDist, groundRange);
+		initJumpDownRig(es, sp, sq, config);
 	}
 	else if (desiredAction == DT_LINK_ACTION_JUMP_OVER)
 	{
diff --git a/Engine/Source/Runtime/Navmesh/Private/Detour/DetourNavLinkBuilderConfig.cpp b/Engine/Source/Runtime/Navmesh/Private/Detour/DetourNavLinkBuilderConfig.cpp
new file mode 100644
--- /dev/null
+++ b/Engine/Source/Runtime/Navmesh/Private/Detour/DetourNavLinkBuilderConfig.cpp
@@ -0,0 +1,22 @@
+﻿// Copyright Epic Games, Inc. All Rights Reserved.
+#include "Detour/DetourNavLinkBuilderConfig.h"
+
+#include "Detour/DetourCommon.h"
+
+void dtNavLinkBuilderJumpDownConfig::init()
+{
+	check(jumpHeight >= 0.f);
+	check(jumpMaxDepth >= 0.f);
+	
+	// Parabolic equation y(x) = ax^2 + (-d/l - al)x
+	// Where 'a' is constant
+	//       'l' is the jump length from the starting point (jumpLength)
+	//       'd' is the distance below the starting point (jumpMaxDepth)
+
+	// Solving a for the jumpHeight (h) from the starting point gives:
+	// a(h) = -(1/l^2) * (d + 2h + 2*srt(h(h+d)))
+	cachedParabolaConstant = -(1/dtSqr(jumpLength)) * (jumpMaxDepth + 2*jumpHeight + 2*dtSqrt(jumpHeight*(jumpHeight+jumpMaxDepth)));
+
+	// Caching this as well as it remains constant for all links
+	cachedDownRatio = -jumpMaxDepth/jumpLength;
+}
\ No newline at end of file
diff --git a/Engine/Source/Runtime/Navmesh/Public/Detour/DetourNavLinkBuilder.h b/Engine/Source/Runtime/Navmesh/Public/Detour/DetourNavLinkBuilder.h
--- a/Engine/Source/Runtime/Navmesh/Public/Detour/DetourNavLinkBuilder.h
+++ b/Engine/Source/Runtime/Navmesh/Public/Detour/DetourNavLinkBuilder.h
@@ -168,8 +168,7 @@ private:
 								   dtReal* outSegs, const int maxOutSegs) const;
 	
 	void initJumpDownRig(EdgeSampler* es, const dtReal* sp, const dtReal* sq,
-						 const float jumpStartDist, const float jumpEndDist,
-						 const float jumpDownDist, const float groundRange) const;
+						 const dtNavLinkBuilderJumpDownConfig& config) const;
 
 	static void initJumpOverRig(EdgeSampler* es, const dtReal* sp, const dtReal* sq,
 						 const float jumpStartDist, const float jumpEndDist,
diff --git a/Engine/Source/Runtime/Navmesh/Public/Detour/DetourNavLinkBuilderConfig.h b/Engine/Source/Runtime/Navmesh/Public/Detour/DetourNavLinkBuilderConfig.h
--- a/Engine/Source/Runtime/Navmesh/Public/Detour/DetourNavLinkBuilderConfig.h
+++ b/Engine/Source/Runtime/Navmesh/Public/Detour/DetourNavLinkBuilderConfig.h
@@ -1,6 +1,8 @@
 ﻿// Copyright Epic Games, Inc. All Rights Reserved.
 #pragma once
 
+#include "CoreTypes.h"
+
 /** Configuration for generated jump down links. */
 struct dtNavLinkBuilderJumpDownConfig
 {
@@ -16,6 +18,9 @@ struct dtNavLinkBuilderJumpDownConfig
 	/// How far below the starting height we want to look for landing ground. [Limit: > 0] [Units: wu]
 	float jumpMaxDepth = 150.f;
 
+	/// Peak height relative to the height of the starting point. [Limit: >= 0] [Units: wu]
+	float jumpHeight = 50.f;
+
 	/// Tolerance at both ends of the jump to find ground. [Limit: > 0] [Units: wu]
 	float jumpEndsHeightTolerance = 50.f;
 
@@ -27,8 +32,17 @@ struct dtNavLinkBuilderJumpDownConfig
 	/// Use greater distance for more filtering (0 to deactivate filtering). [Limit: > 0] [Units: wu]
 	float filterDistanceThreshold = 80.f;
 
+	/// Cached parabola constant fitting the configuration parameters. 
+	float cachedParabolaConstant = 0;
+
+	/// Cached value used when computing jump trajectory.
+	float cachedDownRatio = 0;
+	
 	/// User id used to handle links made from this configuration.
 	unsigned long long linkUserId = 0;
+
+	/// Initialize the configuration by computing cached values.
+	NAVMESH_API void init();
 };
 
 /** Configuration for generated jump over links. */
