[NavLinkGeneration] Account more precisely for the agent radius size when validating link trajectories.
Fix approximation issue at the top of the trajectory in the trajectory initialization.

diff --git a/Engine/Source/Runtime/Navmesh/Private/DebugUtils/DetourNavLinkDebugDraw.cpp b/Engine/Source/Runtime/Navmesh/Private/DebugUtils/DetourNavLinkDebugDraw.cpp
--- a/Engine/Source/Runtime/Navmesh/Private/DebugUtils/DetourNavLinkDebugDraw.cpp
+++ b/Engine/Source/Runtime/Navmesh/Private/DebugUtils/DetourNavLinkDebugDraw.cpp
@@ -104,8 +104,12 @@ namespace UE::Detour::Private
 		}
 	}
 
-	void drawTrajectorySlice(duDebugDraw* dd, const dtReal* pa, const dtReal* pb, const dtNavLinkBuilder::Trajectory2D* trajectory, const unsigned int color)
+	void drawTrajectorySlice(duDebugDraw* dd, dtReal* pa, dtReal* pb, const dtNavLinkBuilder::Trajectory2D* trajectory, const dtReal* trajectoryDir, const unsigned int color)
 	{
+		// Offset start and end points to account for the agent radius.
+		dtVmad(pa, pa, trajectoryDir, -trajectory->radiusOverflow);
+		dtVmad(pb, pb, trajectoryDir,  trajectory->radiusOverflow);
+		
 		unsigned int colt = duTransCol(color, 50);
 		unsigned int colb = duTransCol(duLerpCol(color,duColor::black,96), 50);
 		
@@ -557,9 +561,9 @@ void duDebugDrawNavLinkBuilder(duDebugDraw* dd, const dtNavLinkBuilder& linkBuil
 						ept[1] = esmp->height;
 						
 						if (ssmp->flags & dtNavLinkBuilder::UNRESTRICTED)
-							UE::Detour::Private::drawTrajectorySlice(dd, spt, ept, &es->trajectory, duColor::green);
+							UE::Detour::Private::drawTrajectorySlice(dd, spt, ept, &es->trajectory, es->az, duColor::green);
 						else
-							UE::Detour::Private::drawTrajectorySlice(dd, spt, ept, &es->trajectory, duColor::orangeRed);
+							UE::Detour::Private::drawTrajectorySlice(dd, spt, ept, &es->trajectory, es->az, duColor::orangeRed);
 					}
 					dd->depthMask(true);
 				}
diff --git a/Engine/Source/Runtime/Navmesh/Private/Detour/DetourNavLinkBuilder.cpp b/Engine/Source/Runtime/Navmesh/Private/Detour/DetourNavLinkBuilder.cpp
--- a/Engine/Source/Runtime/Navmesh/Private/Detour/DetourNavLinkBuilder.cpp
+++ b/Engine/Source/Runtime/Navmesh/Private/Detour/DetourNavLinkBuilder.cpp
@@ -527,14 +527,18 @@ bool dtNavLinkBuilder::checkHeightfieldCollision(const dtReal x, const dtReal ym
 }
 
 // Returns true if none of the samples ymin, ymax collide with the heghtfield.
-bool dtNavLinkBuilder::isTrajectoryClear(const dtReal* pa, const dtReal* pb, const Trajectory2D* tra) const
+bool dtNavLinkBuilder::isTrajectoryClear(dtReal* pa, dtReal* pb, const Trajectory2D* trajectory, const dtReal* trajectoryDir) const
 {
-	const int nsamples = tra->samples.Num();
+	// Offset start and end points to account for the agent radius.
+	dtVmad(pa, pa, trajectoryDir, -trajectory->radiusOverflow);
+	dtVmad(pb, pb, trajectoryDir,  trajectory->radiusOverflow);
+
+	const int nsamples = trajectory->samples.Num();
 	const float invLastSample = 1.f / (nsamples-1);
 	for (int i = 0; i < nsamples; ++i)
 	{
 		dtReal p[3];
-		const TrajectorySample& s = tra->samples[i];
+		const TrajectorySample& s = trajectory->samples[i];
 		const float u = (float)i * invLastSample;
 		dtVlerp(p, pa, pb, u);
 		if (checkHeightfieldCollision(p[0], p[1] + s.ymin, p[1] + s.ymax, p[2]))
@@ -600,41 +604,52 @@ void dtNavLinkBuilder::sampleAction(EdgeSampler* es) const
 		spt[1] = ssmp.height;
 		ept[1] = esmp.height;
 		
-		if (!isTrajectoryClear(spt, ept, &es->trajectory))
+		if (!isTrajectoryClear(spt, ept, &es->trajectory, es->az))
 			continue;
 
 		ssmp.flags = static_cast<GroundSampleFlag>((unsigned char)ssmp.flags | (unsigned char)UNRESTRICTED);
 	}
 }
 
-void dtNavLinkBuilder::initTrajectory(Trajectory2D* tra) const
+void dtNavLinkBuilder::initTrajectory(Trajectory2D* trajectory) const
 {
 	using namespace UE::Detour::NavLink::Private;
 	
-	const float* pa = &tra->spine[0];
-	const float* pb = &tra->spine[(tra->nspine-1)*2];
+	const float agentRadius = (float)m_linkBuilderConfig.agentRadius;
+	trajectory->radiusOverflow = agentRadius;
+
+	// Spine points [x,y]. y is up and x is in the direction of the trajectory, relative to the edge. 
+	float pa[2] = { trajectory->spine[0], trajectory->spine[1] };
+	float pb[2] = { trajectory->spine[(trajectory->nspine-1)*2], trajectory->spine[(trajectory->nspine-1)*2+1] };
+
+	// Finding samples along the spine accounting for the agent size,
+	// so we need to look a bit before and after the desired trajectory.
+	pa[0] -= agentRadius;
+	pb[0] += agentRadius;
 	
 	const float dx = pb[0] - pa[0];
 	const int nsamples  = dtMax(2, (int)ceilf(dx*m_invCs));
-	tra->samples.Reserve(nsamples);
+	trajectory->samples.Reserve(nsamples);
+
+	const float* spine = trajectory->spine;
+	unsigned char nspine = trajectory->nspine;
 
 	const unsigned short lastSampleIndex = nsamples-1;
 	const float invLastIndex = 1.f/lastSampleIndex;
 	for (int i = 0; i < nsamples; ++i)
 	{
 		const float u = (float)i * invLastIndex;
-		TrajectorySample& s = tra->samples.Emplace_GetRef();
-		s.x = dtLerp(pa[0], pb[0], u);
+		const float xRef = dtLerp(pa[0], pb[0], u);
+		const float yRef = dtLerp(pa[1], pb[1], u);
 
-		const float* spine = tra->spine;
-		unsigned char nspine = tra->nspine;
-		const float y0 = getHeight(s.x-m_linkBuilderConfig.agentRadius, spine, nspine);
-		const float y1 = getHeight(s.x+m_linkBuilderConfig.agentRadius, spine, nspine);
+		// Sample the height on the spine at 3 locations to get an approximated min and max y.
+		const float y0 = getHeight(xRef - agentRadius, spine, nspine);
+		const float y1 = getHeight(xRef + agentRadius, spine, nspine);
+		const float y2 = getHeight(xRef, spine, nspine);
 
-		const float y = dtLerp(pa[1], pb[1], u);
-		
-		s.ymin = dtMin(y0,y1) + m_linkBuilderConfig.agentClimb - y;	
-		s.ymax = dtMax(y0,y1) + m_linkBuilderConfig.agentHeight - y; 
+		TrajectorySample& s = trajectory->samples.Emplace_GetRef();
+		s.ymin = dtMin(dtMin(y0,y1), y2) + m_linkBuilderConfig.agentClimb - yRef;	
+		s.ymax = dtMax(dtMax(y0,y1), y2) + m_linkBuilderConfig.agentHeight - yRef; 
 	}
 }
 
diff --git a/Engine/Source/Runtime/Navmesh/Public/Detour/DetourNavLinkBuilder.h b/Engine/Source/Runtime/Navmesh/Public/Detour/DetourNavLinkBuilder.h
--- a/Engine/Source/Runtime/Navmesh/Public/Detour/DetourNavLinkBuilder.h
+++ b/Engine/Source/Runtime/Navmesh/Public/Detour/DetourNavLinkBuilder.h
@@ -60,18 +60,19 @@ public:
 	
 	struct TrajectorySample
 	{
-		float x, ymin, ymax;
+		// Min and max heights relative to the height on the segment between the start and the end point (at the sampling locations). 
+		float ymin;
+		float ymax;
 	};
 	
 	struct Trajectory2D
 	{
-		Trajectory2D() : nspine(0) {}
-
 		float spine[2*MAX_SPINE];		// [x,y] relative points representing the desired trajectory (one spine with MAX_SPINE points)
+		float radiusOverflow = 0.f;		// extra distance at the start and end of the spine to lookup for samples
 
 		TArray<TrajectorySample, TInlineAllocator<8>> samples;	// samples along trajectory slices to check for collision
 		
-		unsigned char nspine;			// @todo: remove (use direclty MAX_SPINE) or make relative to trajectory type and config
+		unsigned char nspine = 0;		// @todo: remove (use direclty MAX_SPINE) or make relative to trajectory type and config
 	};
 
 private:	
@@ -160,8 +161,8 @@ public:
 	NAVMESH_API void debugBuildEdge(const dtLinkBuilderConfig& acfg, dtNavLinkAction action, int edgeIndex, EdgeSampler& sampler);
 	
 private:
-	void initTrajectory(Trajectory2D* tra) const;
-	bool isTrajectoryClear(const dtReal* pa, const dtReal* pb, const Trajectory2D* tra) const;
+	void initTrajectory(Trajectory2D* trajectory) const;
+	bool isTrajectoryClear(dtReal* pa, dtReal* pb, const Trajectory2D* trajectory, const dtReal* trajectoryDir) const;
 	
 	int findPotentialJumpOverEdges(const dtReal* sp, const dtReal* sq,
 								   const float depthRange, const float heightRange,
