PR #10587: Functionality Pass for Navigation Segment Links
(with some added logs and changes related to TileRefs)

diff --git a/Engine/Source/Runtime/NavigationSystem/Private/NavMesh/PImplRecastNavMesh.cpp b/Engine/Source/Runtime/NavigationSystem/Private/NavMesh/PImplRecastNavMesh.cpp
--- a/Engine/Source/Runtime/NavigationSystem/Private/NavMesh/PImplRecastNavMesh.cpp
+++ b/Engine/Source/Runtime/NavigationSystem/Private/NavMesh/PImplRecastNavMesh.cpp
@@ -553,7 +553,7 @@ void FPImplRecastNavMesh::Serialize( FArchive& Ar, int32 NavMeshVersion )
 		}
 		
 		NumTiles = TilesToSave.Num();
-		UE_LOG(LogNavigation, VeryVerbose, TEXT("%s Ar.IsSaving() %i tiles in %s."), ANSI_TO_TCHAR(__FUNCTION__), NumTiles, *NavMeshOwner->GetFullName());
+		UE_LOG(LogNavigation, VeryVerbose, TEXT("%hs Ar.IsSaving() %i tiles in %s."), __FUNCTION__, NumTiles, *NavMeshOwner->GetFullName());
 	}
 
 	Ar << NumTiles;
@@ -662,6 +662,13 @@ void FPImplRecastNavMesh::Serialize( FArchive& Ar, int32 NavMeshVersion )
 				UE_VLOG(NavMeshOwner, LogNavigation, Error, TEXT("Failed to initialize NavMesh"));
 			}
 
+#if WITH_NAVMESH_SEGMENT_LINKS			
+			// Store the tileRefs when iterating over the tile list.
+			// These will be processed again to add the segment link connections.
+			TArray<dtTileRef> AddedTiles;
+			AddedTiles.Reserve(NumTiles);
+#endif // WITH_NAVMESH_SEGMENT_LINKS			
+
 			for (int i = 0; i < NumTiles; ++i)
 			{
 				dtTileRef TileRef = MAX_uint64;
@@ -679,12 +686,16 @@ void FPImplRecastNavMesh::Serialize( FArchive& Ar, int32 NavMeshVersion )
 
 				if (TileData != NULL)
 				{
+#if WITH_NAVMESH_SEGMENT_LINKS					
+					AddedTiles.Add(TileRef);
+#endif
+
 					dtMeshHeader* const TileHeader = (dtMeshHeader*)TileData;
 					Status = DetourNavMesh->addTile(TileData, TileDataSize, DT_TILE_FREE_DATA, TileRef, NULL);
 					if (dtStatusDetail(Status, DT_OUT_OF_MEMORY))
 					{
-						UE_LOG(LogNavigation, Warning, TEXT("%s Failed to add tile (%d,%d:%d), %d tile limit reached in %s. If using FixedTilePoolSize, try increasing the TilePoolSize or using bigger tiles."),
-							ANSI_TO_TCHAR(__FUNCTION__), TileHeader->x, TileHeader->y, TileHeader->layer, DetourNavMesh->getMaxTiles(), *NavMeshOwner->GetFullName());
+						UE_LOG(LogNavigation, Warning, TEXT("%hs Failed to add tile (%d,%d:%d), %d tile limit reached in %s. If using FixedTilePoolSize, try increasing the TilePoolSize or using bigger tiles."),
+							__FUNCTION__, TileHeader->x, TileHeader->y, TileHeader->layer, DetourNavMesh->getMaxTiles(), *NavMeshOwner->GetFullName());
 					}
 
 					// Serialize compressed tile cache layer
@@ -699,6 +710,14 @@ void FPImplRecastNavMesh::Serialize( FArchive& Ar, int32 NavMeshVersion )
 					}
 				}
 			}
+
+#if WITH_NAVMESH_SEGMENT_LINKS
+			// Create segment link connections now that all the tiles have been loaded.
+			for (int i = 0; i < AddedTiles.Num(); ++i)
+			{
+				ProcessSegmentLinksForTile(FNavTileRef(AddedTiles[i]));
+			}
+#endif // WITH_NAVMESH_SEGMENT_LINKS
 		}
 	}
 	else if (Ar.IsSaving())
@@ -2098,6 +2117,20 @@ void FPImplRecastNavMesh::UpdateSegmentLinkArea(int32 UserId, uint8 AreaType, ui
 		DetourNavMesh->updateOffMeshSegmentConnectionByUserId(UserId, AreaType, PolyFlags);
 	}
 }
+
+void FPImplRecastNavMesh::ProcessSegmentLinksForTile(FNavTileRef TileRef) const
+{
+	uint32 NumSkippedNeighborTiles;
+	ProcessSegmentLinksForTile(TileRef, 0 /*MaxSkippedNeigborTiles*/, nullptr /*OutSkippedNeigborTiles*/, NumSkippedNeighborTiles);
+}
+
+void FPImplRecastNavMesh::ProcessSegmentLinksForTile(FNavTileRef TileRef, uint32 MaxSkippedNeigborTiles, dtTileRef* OutSkippedNeigborTiles, uint32& OutNumSkippedNeigborTiles) const
+{
+	if (DetourNavMesh)
+	{
+		DetourNavMesh->processSegmentLinksForTile(static_cast<dtTileRef>(TileRef), MaxSkippedNeigborTiles, OutSkippedNeigborTiles, OutNumSkippedNeigborTiles);
+	}
+}
 #endif // WITH_NAVMESH_SEGMENT_LINKS
 
 bool FPImplRecastNavMesh::GetPolyCenter(NavNodeRef PolyID, FVector& OutCenter) const
@@ -2155,6 +2188,16 @@ bool FPImplRecastNavMesh::GetPolyVerts(NavNodeRef PolyID, TArray<FVector>& OutVe
 				OutVerts.Add( Recast2UnrVector(V) );
 			}
 
+#if WITH_NAVMESH_SEGMENT_LINKS
+			// Segment polys store each segment edge from left to right.
+			// Therefore, swap the edge stored at indices 2 and 3 to ensure the poly verts
+			// are in proper clockwise order for drawing to avoid twisted segment link polys.
+			if (Poly->getType() == DT_POLYTYPE_OFFMESH_SEGMENT)
+			{
+				OutVerts.SwapMemory(2, 3);
+			}
+#endif // WITH_NAVMESH_SEGMENT_LINKS
+
 			return true;
 		}
 	}
@@ -2927,8 +2970,8 @@ int32 FPImplRecastNavMesh::GetTilesDebugGeometry(const FRecastNavMeshGenerator*
 		if (bIsBeingBuilt)
 		{
 			const dtTileRef TileRef = DetourNavMesh->getTileRef(&Tile);
-			UE_LOG(LogNavigation, VeryVerbose, TEXT("%s TileId: %d Salt: %d TileRef: 0x%llx bIsBeingBuilt"),
-				ANSI_TO_TCHAR(__FUNCTION__), DetourNavMesh->decodePolyIdTile(TileRef), DetourNavMesh->decodePolyIdSalt(TileRef), TileRef);	
+			UE_LOG(LogNavigation, VeryVerbose, TEXT("%hs TileId: %d Salt: %d TileRef: 0x%llx bIsBeingBuilt"),
+				__FUNCTION__, DetourNavMesh->decodePolyIdTile(TileRef), DetourNavMesh->decodePolyIdSalt(TileRef), TileRef);	
 		}
 	});
 
diff --git a/Engine/Source/Runtime/NavigationSystem/Private/NavMesh/RecastNavMesh.cpp b/Engine/Source/Runtime/NavigationSystem/Private/NavMesh/RecastNavMesh.cpp
--- a/Engine/Source/Runtime/NavigationSystem/Private/NavMesh/RecastNavMesh.cpp
+++ b/Engine/Source/Runtime/NavigationSystem/Private/NavMesh/RecastNavMesh.cpp
@@ -497,6 +497,7 @@ ARecastNavMesh::ARecastNavMesh(const FObjectInitializer& ObjectInitializer)
 	MaxSimultaneousTileGenerationJobsCount = 1024;
 	bDoFullyAsyncNavDataGathering = false;
 	TileNumberHardLimit = 1 << 20;
+	ExpectedMaxLayersPerTile = 12;
 
 #if RECAST_ASYNC_REBUILDING
 	BatchQueryCounter = 0;
@@ -2760,6 +2761,86 @@ void ARecastNavMesh::OnNavMeshTilesUpdated(const TArray<FNavTileRef>& ChangedTil
 	InvalidateAffectedPaths(ChangedTiles);
 }
 
+#if WITH_NAVMESH_SEGMENT_LINKS
+void ARecastNavMesh::CreateSegmentLinkConnections(TSet<FNavTileRef>& TilesPendingSegmentLinks)
+{
+	if(!RecastNavMeshImpl || !RecastNavMeshImpl->DetourNavMesh)
+	{
+		return;
+	}
+	
+	TSet<FNavTileRef> ProcessedTiles;
+	ProcessedTiles.Reserve(TilesPendingSegmentLinks.Num());
+
+	// Segment link connections are created after all the navmesh tiles have been generated.
+	// This means care needs to be taken to ensure links don't get created twice. For example,
+	// linking together TileA and TileB but also TileB and TileA in the opposite order.
+	// Therefore, sort the list of tiles needing link connections.
+	// Tiles can only link to other tiles with a greater TileIndex to ensure the segment link connections are not duplicated.
+	TArray<FNavTileRef> TilesPendingSegmentLinksArray = TilesPendingSegmentLinks.Array();
+
+	auto TileSorter = [this](const FNavTileRef A, const FNavTileRef B)
+	{
+		const dtMeshTile* TileA = RecastNavMeshImpl->DetourNavMesh->getTileByRef(static_cast<dtTileRef>(A));
+		const dtMeshTile* TileB = RecastNavMeshImpl->DetourNavMesh->getTileByRef(static_cast<dtTileRef>(B));
+		return	TileA && TileA->header &&
+				TileB && TileB->header &&
+				(TileA->header->x < TileB->header->x ||
+				TileA->header->y < TileB->header->y ||
+				TileA->header->layer < TileB->header->layer);	
+	};
+	TilesPendingSegmentLinksArray.Sort(TileSorter);
+
+	for (FNavTileRef TileRef : TilesPendingSegmentLinksArray)
+	{
+		if (ProcessedTiles.Contains(TileRef))
+		{
+			continue;
+		}
+
+		// We need to know what tiles were skipped due to those tiles having a lower index than the tile
+		// being processed. These tiles may need to be processed if they were not in the original list
+		// of tiles pending segment links.
+		// 
+		// All 8 tile neighbors could potentially have a lower tile index than the source tile and will be skipped.
+		// The skipped neighbors array is pre-allocated for use with the Detour library.
+		constexpr uint32 MaxNeighborsToBeSkipped = 8;
+		const uint32 MaxSkippedNeighborTiles = MaxNeighborsToBeSkipped * ExpectedMaxLayersPerTile;
+
+		TArray<dtTileRef> SkippedNeighborTiles;
+		SkippedNeighborTiles.AddUninitialized(MaxSkippedNeighborTiles);
+
+		// Process the tile and store the neighbor tiles that were skipped. They will need to be processed if they haven't been already.
+		uint32 NumSkippedNeighborTiles = 0;
+		UE_LOG(LogSegmentLink, Log, TEXT("1) %hs and store skipped tiles"), __FUNCTION__);
+		RecastNavMeshImpl->ProcessSegmentLinksForTile(TileRef, MaxSkippedNeighborTiles, SkippedNeighborTiles.GetData(), NumSkippedNeighborTiles);
+		ProcessedTiles.Add(TileRef);
+
+		if (NumSkippedNeighborTiles > MaxSkippedNeighborTiles)
+		{
+			UE_LOG(LogNavigation, Warning, TEXT("The number of neighbor tiles skipped (%d) while connecting segment links exceeds the max allowed (%d). Consider increasing ExpectedMaxLayersPerTile in the NavMesh settings."), NumSkippedNeighborTiles, MaxSkippedNeighborTiles);
+
+			NumSkippedNeighborTiles = MaxSkippedNeighborTiles;
+		}
+
+		// Process skipped tiles if they have not been already to get full link coverage.
+		for (uint32 Index = 0; Index < NumSkippedNeighborTiles; ++Index)
+		{
+			const FNavTileRef NeighborTileRef = FNavTileRef(SkippedNeighborTiles[Index]);
+			if (!ProcessedTiles.Contains(NeighborTileRef))
+			{
+				// No need to worry about skipped neighbors here. This tile was not in the original update
+				// list and only needs to be processed to connect the links in the previously processed tile.
+				UE_LOG(LogSegmentLink, Log, TEXT("2) %hs process skipped tileRef %llu"), __FUNCTION__, uint64(NeighborTileRef));
+				RecastNavMeshImpl->ProcessSegmentLinksForTile(NeighborTileRef);
+
+				ProcessedTiles.Add(NeighborTileRef);
+			}
+		}
+	}
+}
+#endif // WITH_NAVMESH_SEGMENT_LINKS
+
 void ARecastNavMesh::InvalidateAffectedPaths(const TArray<FNavTileRef>& ChangedTiles)
 {
 	const int32 PathsCount = ActivePaths.Num();
diff --git a/Engine/Source/Runtime/NavigationSystem/Private/NavMesh/RecastNavMeshDataChunk.cpp b/Engine/Source/Runtime/NavigationSystem/Private/NavMesh/RecastNavMeshDataChunk.cpp
--- a/Engine/Source/Runtime/NavigationSystem/Private/NavMesh/RecastNavMeshDataChunk.cpp
+++ b/Engine/Source/Runtime/NavigationSystem/Private/NavMesh/RecastNavMeshDataChunk.cpp
@@ -327,6 +327,17 @@ TArray<FNavTileRef> URecastNavMeshDataChunk::AttachTiles(ARecastNavMesh& NavMesh
 				Result.Add(FNavTileRef(TileRef));
 			}
 		}
+
+#if WITH_NAVMESH_SEGMENT_LINKS
+		// Create segment link connections now that all the tiles have been loaded.
+		if (const FPImplRecastNavMesh* const NavMeshImpl = NavMesh.GetRecastNavMeshImpl())
+		{
+			for (int32 Index = 0; Index < Result.Num(); ++Index)
+			{
+				NavMeshImpl->ProcessSegmentLinksForTile(Result[Index]);
+			}
+		}
+#endif // WITH_NAVMESH_SEGMENT_LINKS		
 	}
 
 	UE_LOG(LogNavigation, Verbose, TEXT("Attached %d tiles to NavMesh - %s"), Result.Num(), *NavigationDataName.ToString());
diff --git a/Engine/Source/Runtime/NavigationSystem/Private/NavMesh/RecastNavMeshGenerator.cpp b/Engine/Source/Runtime/NavigationSystem/Private/NavMesh/RecastNavMeshGenerator.cpp
--- a/Engine/Source/Runtime/NavigationSystem/Private/NavMesh/RecastNavMeshGenerator.cpp
+++ b/Engine/Source/Runtime/NavigationSystem/Private/NavMesh/RecastNavMeshGenerator.cpp
@@ -1402,7 +1402,8 @@ struct FOffMeshData
 			});
 		}
 	}
-	
+
+#if WITH_NAVMESH_SEGMENT_LINKS	
 	void AddSegmentLinks(const TArray<FNavigationSegmentLink>& Links, const FTransform& LocalToWorld, int32 AgentIndex, float DefaultSnapHeight)
 	{
 		for (int32 LinkIndex = 0; LinkIndex < Links.Num(); ++LinkIndex)
@@ -1442,6 +1443,7 @@ struct FOffMeshData
 			LinkParams.Add(NewInfo);
 		}
 	}
+#endif // WITH_NAVMESH_SEGMENT_LINKS
 
 protected:
 
@@ -7173,6 +7175,18 @@ TArray<FNavTileRef> FRecastNavMeshGenerator::ProcessTileTasksAndGetUpdatedTiles(
 	}
 #endif
 
+#if WITH_NAVMESH_SEGMENT_LINKS
+	const dtNavMesh* DetourMesh = DestNavMesh->GetRecastNavMeshImpl()->GetRecastMesh();
+	if (DetourMesh)
+	{
+		TilesPendingSegmentLinkConnections.Reserve(TilesPendingSegmentLinkConnections.Num() + UpdatedTiles.Num());
+		for (FNavTileRef TileRef : UpdatedTiles)
+		{
+			TilesPendingSegmentLinkConnections.Add(TileRef);
+		}
+	}
+#endif
+
 	// Notify owner in case all tasks has been completed
 	const bool bHasTasksAtEnd = GetNumRemaningBuildTasks() > 0;
 	if (bHasTasksAtStart && !bHasTasksAtEnd)
@@ -7184,7 +7198,12 @@ TArray<FNavTileRef> FRecastNavMeshGenerator::ProcessTileTasksAndGetUpdatedTiles(
 			UE_LOG(LogNavigationDataBuild, Display, TEXT("   %hs build time: %.2fs"), __FUNCTION__, (FPlatformTime::Seconds() - RebuildAllStartTime));
 			RebuildAllStartTime = 0;
 		}
-		
+
+#if WITH_NAVMESH_SEGMENT_LINKS
+		DestNavMesh->CreateSegmentLinkConnections(TilesPendingSegmentLinkConnections);
+		TilesPendingSegmentLinkConnections.Empty();
+#endif
+
 		DestNavMesh->OnNavMeshGenerationFinished();
 	}
 
diff --git a/Engine/Source/Runtime/NavigationSystem/Public/NavMesh/PImplRecastNavMesh.h b/Engine/Source/Runtime/NavigationSystem/Public/NavMesh/PImplRecastNavMesh.h
--- a/Engine/Source/Runtime/NavigationSystem/Public/NavMesh/PImplRecastNavMesh.h
+++ b/Engine/Source/Runtime/NavigationSystem/Public/NavMesh/PImplRecastNavMesh.h
@@ -172,6 +172,12 @@ public:
 #if WITH_NAVMESH_SEGMENT_LINKS
 	/** Updates area on polygons creating segment-to-segment connection with given UserId */
 	NAVIGATIONSYSTEM_API void UpdateSegmentLinkArea(int32 UserId, uint8 AreaType, uint16 PolyFlags) const;
+
+	/** Creates segment link connections for the give tile */
+	NAVIGATIONSYSTEM_API void ProcessSegmentLinksForTile(FNavTileRef TileRef) const;
+
+	/** Creates segment link connections for the give tile and stores the list of indices skipped to avoid generating duplicate links */
+	NAVIGATIONSYSTEM_API void ProcessSegmentLinksForTile(FNavTileRef TileRef, uint32 MaxSkippedNeigborTiles, dtTileRef* OutSkippedNeigborTiles, uint32& OutNumSkippedNeigborTiles) const;
 #endif // WITH_NAVMESH_SEGMENT_LINKS
 
 	/** Retrieves center of the specified polygon. Returns false on error. */
diff --git a/Engine/Source/Runtime/NavigationSystem/Public/NavMesh/RecastNavMesh.h b/Engine/Source/Runtime/NavigationSystem/Public/NavMesh/RecastNavMesh.h
--- a/Engine/Source/Runtime/NavigationSystem/Public/NavMesh/RecastNavMesh.h
+++ b/Engine/Source/Runtime/NavigationSystem/Public/NavMesh/RecastNavMesh.h
@@ -734,6 +734,10 @@ class ARecastNavMesh : public ANavigationData
 	UPROPERTY(EditAnywhere, Category = Generation, config, meta = (ClampMin = "1", UIMin = "1"), AdvancedDisplay)
 	int32 TileNumberHardLimit;
 
+	/** Used when connecting segment links across layers to determine how much memory to allocate to hold skipped layers */
+	UPROPERTY(EditAnywhere, Category = Generation, config, meta = (ClampMin = "1", UIMin = "1"), AdvancedDisplay)
+	int32 ExpectedMaxLayersPerTile;
+
 	UPROPERTY(VisibleAnywhere, Category = Generation, AdvancedDisplay)
 	int32 PolyRefTileBits;
 
@@ -1145,6 +1149,11 @@ public:
 	/** called after regenerating tiles */
 	NAVIGATIONSYSTEM_API virtual void OnNavMeshTilesUpdated(const TArray<FNavTileRef>& ChangedTiles);
 
+#if WITH_NAVMESH_SEGMENT_LINKS
+	/** Creates segment link connections for the given tiles. Called after navigation tiles are finished generating. */
+	NAVIGATIONSYSTEM_API virtual void CreateSegmentLinkConnections(TSet<FNavTileRef>& TilesPendingSegmentLinks);
+#endif
+
 	/** Event from generator that navmesh build has finished */
 	NAVIGATIONSYSTEM_API virtual void OnNavMeshGenerationFinished();
 
diff --git a/Engine/Source/Runtime/NavigationSystem/Public/NavMesh/RecastNavMeshGenerator.h b/Engine/Source/Runtime/NavigationSystem/Public/NavMesh/RecastNavMeshGenerator.h
--- a/Engine/Source/Runtime/NavigationSystem/Public/NavMesh/RecastNavMeshGenerator.h
+++ b/Engine/Source/Runtime/NavigationSystem/Public/NavMesh/RecastNavMeshGenerator.h
@@ -1129,6 +1129,11 @@ protected:
 	/** List of dirty tiles currently being regenerated */
 	TNavStatArray<TRunningTileElement<FRecastTileGeneratorWrapper>> RunningDirtyTiles;
 
+#if WITH_NAVMESH_SEGMENT_LINKS
+	/** List of tiles that have been generated but have not completed segment link connections yet */
+	TSet<FNavTileRef> TilesPendingSegmentLinkConnections;
+#endif // WITH_NAVMESH_SEGMENT_LINKS
+
 #if WITH_EDITOR
 	/** List of tiles that were recently regenerated */
 	TNavStatArray<FTileTimestamp> RecentlyBuiltTiles;
diff --git a/Engine/Source/Runtime/Navmesh/Private/Detour/DetourNavMesh.cpp b/Engine/Source/Runtime/Navmesh/Private/Detour/DetourNavMesh.cpp
--- a/Engine/Source/Runtime/Navmesh/Private/Detour/DetourNavMesh.cpp
+++ b/Engine/Source/Runtime/Navmesh/Private/Detour/DetourNavMesh.cpp
@@ -25,6 +25,10 @@
 
 DEFINE_LOG_CATEGORY(LogDetour);
 
+#if WITH_NAVMESH_SEGMENT_LINKS
+DEFINE_LOG_CATEGORY(LogSegmentLink);
+#endif
+
 //@UE BEGIN Adding support for memory tracking.
 static dtStatsPostAddTileFunc* sAddTileFunc = nullptr;
 static dtStatsPreRemoveTileFunc* sRemoveTileFunc = nullptr;
@@ -346,15 +350,26 @@ struct dtOffMeshSegmentTileIntersection
 
 struct dtOffMeshSegmentIntersectionLink
 {
-	dtReal t;
+	// Links connect intersection points on different segments so the t values for both edge
+	// segments are needed to retrieve the correct intersection points later.
+	dtReal tA, tB;
 	unsigned int polyA, polyB;
 	dtMeshTile* tileA;
 	dtMeshTile* tileB;
 };
 
+// This class stores pairs of intersection links.
+// The pairs found will be sorted by relative t values and then converted into segment parts in order.
+struct dtOffMeshSegmentIntersectionLinkPair
+{
+	dtOffMeshSegmentIntersectionLink links[2];
+};
+
 struct dtOffMeshSegmentPart
 {
-	dtReal t0, t1;
+	// Separate t values for A and B edges.
+	// Parts are formed from two unique segments requiring their own t values to retrieve the correct end points from the original segments.
+	dtReal tA0, tA1, tB0, tB1;
 	unsigned short vA0, vA1, vB0, vB1;
 	unsigned int polyA, polyB;
 	dtMeshTile* tileA;
@@ -384,30 +399,40 @@ static bool isPolyIntersectingSegment(const dtMeshTile* tile, int polyIdx,
 	return dtIntersectSegmentPoly2D(spos, epos, verts, poly->vertCount, tmin, tmax, smin, smax);
 }
 
-static void addSegmentIntersections(dtOffMeshSegmentIntersection* isec, dtOffMeshSegmentIntersection* list, int& nlist)
+static void addSegmentIntersections(const dtOffMeshSegmentIntersection* isec, dtOffMeshSegmentIntersection* list, int& nlist)
 {
+	UE_LOG(LogSegmentLink, Log, TEXT("         %hs"), __FUNCTION__);
+	
 	int minIdx = -1;
 	int maxIdx = -1;
 	for (int i = 0; i < nlist; i++)
 	{
-		if (isIntersectionPointEqual(list[i].t, isec[0].t)) minIdx = i;
-		if (isIntersectionPointEqual(list[i].t, isec[1].t)) maxIdx = i;
+		if (isIntersectionPointEqual(list[i].t, isec[0].t))
+			minIdx = i;
+		
+		if (isIntersectionPointEqual(list[i].t, isec[1].t))
+			maxIdx = i;
 	}
 
 	// min = overwrite if exists or add new one
 	if (minIdx < 0)
 	{
+		UE_LOG(LogSegmentLink, Log, TEXT("            %hs minIdx=%i, add seg intersection, tile (%i,%i,%i), poly %i, %f"),
+			__FUNCTION__, minIdx, isec[0].tile->header->x, isec[0].tile->header->y, isec[0].tile->header->layer, isec[0].poly, isec[0].t);
 		list[nlist] = isec[0];
 		nlist++;
 	}
 	else
 	{
+		UE_LOG(LogSegmentLink, Log, TEXT("            %hs override"), __FUNCTION__);
 		list[minIdx] = isec[0];
 	}
 
 	// max = skip if exists or add new one
 	if (maxIdx < 0 && nlist < DT_MAX_OFFMESH_SEGMENT_POINTS)
 	{
+		UE_LOG(LogSegmentLink, Log, TEXT("            %hs maxIdx=%i, add seg intersection, tile (%i,%i,%i), poly %i, %f"),
+			__FUNCTION__, maxIdx, isec[1].tile->header->x, isec[1].tile->header->y, isec[1].tile->header->layer, isec[1].poly, isec[1].t);
 		list[nlist] = isec[1];
 		nlist++;
 	}
@@ -426,15 +451,21 @@ static void gatherSegmentIntersections(dtMeshTile* tile,
 	dtOffMeshSegmentTileIntersection& list)
 {
 	// get all polys intersecting with segment
-	dtReal segBMin[3], segBMax[3], segRad[3] = { radius, walkableClimb, radius };
-	dtVcopy(segBMin, spos);
-	dtVcopy(segBMax, spos);
-	dtVmin(segBMin, epos);
-	dtVmax(segBMax, epos);
-	dtVsub(segBMin, segBMin, segRad);
-	dtVadd(segBMax, segBMax, segRad);
-
-	if (!dtOverlapBounds(segBMin, segBMax, tile->header->bmin, tile->header->bmax))
+	
+	UE_LOG(LogSegmentLink, Log, TEXT("      %hs for (%i,%i,%i) segment: start (%0.2f, %0.2f, %0.2f) end (%0.2f, %0.2f, %0.2f)"),
+		__FUNCTION__, tile->header->x, tile->header->y, tile->header->layer,
+		spos[0], spos[1], spos[2], epos[0], epos[1], epos[2]);
+
+	// compute bounds including radius
+	dtReal segBoundsMin[3], segBoundsMax[3], segRad[3] = { radius, walkableClimb, radius };
+	dtVcopy(segBoundsMin, spos);
+	dtVcopy(segBoundsMax, spos);
+	dtVmin(segBoundsMin, epos);
+	dtVmax(segBoundsMax, epos);
+	dtVsub(segBoundsMin, segBoundsMin, segRad);
+	dtVadd(segBoundsMax, segBoundsMax, segRad);
+
+	if (!dtOverlapBounds(segBoundsMin, segBoundsMax, tile->header->bmin, tile->header->bmax))
 		return;
 
 	dtOffMeshSegmentIntersection intersec[2];
@@ -454,7 +485,7 @@ static void gatherSegmentIntersections(dtMeshTile* tile,
 		}
 
 		// use simple AABB overlap test first
-		if (dtOverlapBounds(segBMin, segBMax, bmin, bmax))
+		if (dtOverlapBounds(segBoundsMin, segBoundsMax, bmin, bmax))
 		{
 			// mark intersection
 			if (isPolyIntersectingSegment(tile, i, spos, epos, intersec[0].t, intersec[1].t))
@@ -474,12 +505,14 @@ static dtOffMeshSegmentData* initSegmentIntersection(const dtNavMesh& navMesh, d
 {
 	const int segCount = tile->header->offMeshSegConCount;
 	if (segCount <= 0)
-		return NULL;
+		return nullptr;
 
 	dtOffMeshSegmentData* segs = (dtOffMeshSegmentData*)dtAlloc(sizeof(dtOffMeshSegmentData)*segCount, DT_ALLOC_TEMP);
-	if (segs == NULL)
-		return NULL;
+	if (segs == nullptr)
+		return nullptr;
 
+	UE_LOG(LogSegmentLink, Log, TEXT("   %hs for tile (%i,%i,%i)"), __FUNCTION__, tile->header->x, tile->header->y, tile->header->layer);
+	
 	memset(segs, 0, sizeof(dtOffMeshSegmentData)*segCount);
 	for (int i = 0; i < segCount; i++)
 	{
@@ -493,9 +526,9 @@ static dtOffMeshSegmentData* initSegmentIntersection(const dtNavMesh& navMesh, d
 	return segs;
 }
 
-static void appendSegmentIntersection(const dtNavMesh& navMesh, dtOffMeshSegmentData* seg, dtMeshTile* tile, dtMeshTile* nei)
+static void appendSegmentIntersection(const dtNavMesh& navMesh, dtOffMeshSegmentData* seg, const dtMeshTile* tile, dtMeshTile* nei)
 {
-	if (seg == NULL)
+	if (seg == nullptr)
 		return;
 
 	for (int i = 0; i < tile->header->offMeshSegConCount; i++)
@@ -507,146 +540,228 @@ static void appendSegmentIntersection(const dtNavMesh& navMesh, dtOffMeshSegment
 	}
 }
 
-int segmentIntersectionLinkSorter(const void* i1, const void* i2)
+int segmentIntersectionLinkPairSorter(const void* i1, const void* i2)
 {
-	// lesser "t" goes first
-	const dtOffMeshSegmentIntersectionLink* a = (const dtOffMeshSegmentIntersectionLink*)i1;
-	const dtOffMeshSegmentIntersectionLink* b = (const dtOffMeshSegmentIntersectionLink*)i2;
-	return (a->t < b->t) ? -1 : (a->t > b->t) ? 1 : 0;
+	// lesser combined "t" goes first
+	const dtOffMeshSegmentIntersectionLinkPair* first = (const dtOffMeshSegmentIntersectionLinkPair*)i1;
+	const dtOffMeshSegmentIntersectionLinkPair* second = (const dtOffMeshSegmentIntersectionLinkPair*)i2;
+
+	const int tFirst = (first->links[1].tA - first->links[0].tA) + (first->links[1].tB - first->links[0].tB);
+	const int tSecond = (second->links[1].tA - second->links[0].tA) + (second->links[1].tB - second->links[0].tB);
+
+	return (tFirst < tSecond) ? -1 : (tFirst > tSecond) ? 1 : 0;
 }
 
 int segmentPartSorter(const void* i1, const void* i2)
 {
-	// higher length (t1-t0) goes first
+	// higher combined relative edge length ((tA1-tA0) + (tB1-tB0)) goes first
 	const dtOffMeshSegmentPart* a = (const dtOffMeshSegmentPart*)i1;
 	const dtOffMeshSegmentPart* b = (const dtOffMeshSegmentPart*)i2;
-	const dtReal lenA = a->t1 - a->t0;
-	const dtReal lenB = b->t1 - b->t0;
+	const dtReal lenA = (a->tA1 - a->tA0) + (a->tB1 - a->tB0);
+	const dtReal lenB = (b->tA1 - b->tA0) + (b->tB1 - b->tB0);
 	return (lenA > lenB) ? -1 : (lenA < lenB) ? 1 : 0;
 }
 
-static unsigned int findMatchingSegmentIntersection(const dtReal t, const dtOffMeshSegmentIntersection* points, const int npoints, bool bAllowExisting)
+// Advance the index into the specified list skipping over sections of the segment that do not intersect any polygons.
+// 
+// In this image, intersection points are defined with an X. The first two intersection points are owned by PolyA. The fact
+// that both points are owned by the same poly signals that there is no polygon that shares the edge where the intersection occurs.
+// If PolyA and PolyB shared the edge where the second intersection occurred, then PolyB would be the owner of the second
+// intersection point and we would know that there is no gap.
+// 
+//          PolyA           Gap           PolyB
+//      ##############               ##############
+//      #            #               #            #
+//      #            #               #            #
+//      #            # ----Skip----> #            #
+//  ----X------------X---------------X------------X---->
+//      #            #               #            #
+//      #            #               #            #
+//      ##############               ##############
+//
+static void advanceListIndex(int& listIndex, const dtOffMeshSegmentTileIntersection& list)
 {
-	if (npoints < 1 || t < points[0].t || t > points[npoints - 1].t)
-		return DT_INVALID_SEGMENT;
-
-	unsigned int seg = DT_INVALID_SEGMENT;
-	for (int i = 1; i < npoints; i++)
+	unsigned int lastPoly = list.points[listIndex].poly;
+	const dtMeshTile* lastTile = list.points[listIndex].tile;
+	if (++listIndex < list.npoints)
 	{
-		if (t <= points[i].t)
+		unsigned int nextPoly = list.points[listIndex].poly;
+		const dtMeshTile* nextTile = list.points[listIndex].tile;
+		// The next intersection point will be owned by a different poly if the segment continues directly into another polygon on the navmesh.
+		// Otherwise, there is a gap and we should not link any part of a segment that intersects a gap.
+		if (lastPoly == nextPoly && lastTile == nextTile)
 		{
-			if (bAllowExisting || (!isIntersectionPointEqual(t, points[i].t) && !isIntersectionPointEqual(t, points[i-1].t)))
-			{
-				if (i < 2 || points[i - 2].poly != points[i - 1].poly)
-					seg = i - 1;
-			}
-
-			return seg;
+			// This there is a gap between the previous poly intersected and the next poly intersected by this segment.
+			// Advance one more index to get to the next poly that is intersected.
+			++listIndex;
 		}
 	}
-
-	return seg;
 }
 
-static bool canConnectSegmentPart(unsigned int startPoly, unsigned int endPoly,
-	const dtMeshTile* startTile, const dtMeshTile* endTile,
-	const dtOffMeshSegmentIntersection* points, const int npoints)
+// This function matches points in listA with points in listB that can form a pair of links that can later be converted into segment parts.
+// Points in the longer list will be shared with other pairs at most once, but points in the shorter list might be shared as many times
+// as it takes to create a link to every point in the shorter list.
+// The number of link pairs found is returned.
+static int findSegmentIntersectionLinkPairs(const dtOffMeshSegmentTileIntersection& listA, const dtOffMeshSegmentTileIntersection& listB, dtOffMeshSegmentIntersectionLinkPair* outLinkPairs)
 {
-	if ((startPoly != endPoly || startTile != endTile) && (npoints > 1))
+	// Each points array needs at least two points to create at least one pair of segment links.
+	if (listA.npoints < 2 || listB.npoints < 2)
+	{
+		return 0;
+	}
+
+	int nPairs = 0;
+	int indexA = 0;
+	int indexB = 0;
+
+	// Determine aliases to A and B specific data based on the length of the lists involved.
+	bool bListAisLonger = listA.npoints > listB.npoints;
+	const dtOffMeshSegmentTileIntersection& longerList = bListAisLonger ? listA : listB;
+	const dtOffMeshSegmentTileIntersection& shorterList = bListAisLonger ? listB : listA;
+	int& longerListIndex = bListAisLonger ? indexA : indexB;
+	int& shorterListIndex = bListAisLonger ? indexB : indexA;
+
+	// Lambda function for convenience to create link pairs.
+	auto createSegmentIntersectionLinkPair = [&outLinkPairs, &nPairs, &listA, &listB, &indexA, &indexB]()
+	{
+		dtOffMeshSegmentIntersectionLink& firstLink = outLinkPairs[nPairs].links[0];
+		firstLink.tA = listA.points[indexA].t;
+		firstLink.polyA = listA.points[indexA].poly;
+		firstLink.tileA = listA.points[indexA].tile;
+		firstLink.tB = listB.points[indexB].t;
+		firstLink.polyB = listB.points[indexB].poly;
+		firstLink.tileB = listB.points[indexB].tile;
+
+		dtOffMeshSegmentIntersectionLink& secondLink = outLinkPairs[nPairs].links[1];
+		secondLink.tA = listA.points[indexA + 1].t;
+		secondLink.polyA = listA.points[indexA + 1].poly;
+		secondLink.tileA = listA.points[indexA + 1].tile;
+		secondLink.tB = listB.points[indexB + 1].t;
+		secondLink.polyB = listB.points[indexB + 1].poly;
+		secondLink.tileB = listB.points[indexB + 1].tile;
+
+		nPairs++;
+	};
+
+	// Create the first pair of links which is guaranteed to connect the first two points in listA with the first two points in listB.
+	createSegmentIntersectionLinkPair();
+
+	// Advance to the next point in the longer list looking for the next pair of links.
+	advanceListIndex(longerListIndex, longerList);
+
+	// Iterate over the longer list since it is known that the pair of points used in the longer list will increment after each segment link pair is created.
+	while (longerListIndex < longerList.npoints - 1)
 	{
-		for (int i = 1; i < npoints; i++)
+		// Calculate the number of points remaining in each list that have not been considered for pairing yet.
+		const int numRemainingLongListPoints = longerList.npoints - longerListIndex - 2;
+		const int numRemainingShortListPoints = shorterList.npoints - shorterListIndex - 2;
+
+		// Determine if the index to the pair of points used in the shorter list should increment based on
+		// which points in the shorter list have closer t values to the current points in the longer list.
+		if (numRemainingShortListPoints > 0)
 		{
-			if (points[i-1].poly == points[i].poly &&
-				points[i-1].tile == points[i].tile &&
-				startPoly == points[i].poly &&
-				startTile == points[i].tile)
+			// Ensure every point in the shorter list also gets a connection to a point in the longer list regardless of t values.
+			if (numRemainingShortListPoints > numRemainingLongListPoints)
 			{
-				return false;
+				advanceListIndex(shorterListIndex, shorterList);
+			}
+			else
+			{
+				int currentTDiff = dtAbs(longerList.points[longerListIndex].t - shorterList.points[shorterListIndex].t) +
+					dtAbs(longerList.points[longerListIndex + 1].t - shorterList.points[shorterListIndex + 1].t);
+
+				int nextTDiff = dtAbs(longerList.points[longerListIndex].t - shorterList.points[shorterListIndex + 1].t) +
+					dtAbs(longerList.points[longerListIndex + 1].t - shorterList.points[shorterListIndex + 2].t);
+
+				if (nextTDiff < currentTDiff)
+				{
+					// The next pair of points in the shorter list are more appropriate to link to the current points in the longer list.
+					advanceListIndex(shorterListIndex, shorterList);
+				}
 			}
 		}
+
+		// Create the next pair of segment intersection links based on the currently selected points in both lists.
+		createSegmentIntersectionLinkPair();
+
+		// Advance to the next point in the longer list looking for the next pair of links.
+		advanceListIndex(longerListIndex, longerList);
 	}
 
-	return true;
+	return nPairs;
 }
 
-static unsigned short findOrAddUniqueValue(dtReal v, dtReal* arr, unsigned short& narr)
+static unsigned short findOrAddUniqueValue(dtReal value, dtReal* array, unsigned short& narray)
 {
-	for (unsigned short i = 0; i < narr; i++)
+	for (unsigned short i = 0; i < narray; i++)
 	{
-		if (arr[i] == v)
+		if (array[i] == value)
 			return i;
 	}
 
-	const unsigned short pos = narr;
-	arr[pos] = v;
-	narr++;
+	const unsigned short pos = narray;
+	array[pos] = value;
+	narray++;
 	return pos;
 }
 
-static void createSegmentParts(dtMeshTile* tile, const dtOffMeshSegmentData& segData,
+static void createSegmentParts(const dtMeshTile* tile, const dtOffMeshSegmentData& segData,
 	dtOffMeshSegmentPart* parts, const int maxParts, int& nparts, unsigned short& nverts)
 {
 	if (segData.listA.npoints <= 0 && segData.listB.npoints <= 0)
 		return;
 
-	const int maxLinks = DT_MAX_OFFMESH_SEGMENT_POINTS * 2;
-	dtOffMeshSegmentIntersectionLink links[maxLinks];
-	memset(links, 0, sizeof(dtOffMeshSegmentIntersectionLink)*maxLinks);
-	int nlinks = 0;
+	UE_LOG(LogSegmentLink, Log, TEXT("   %hs (%i,%i,%i)"), __FUNCTION__, tile->header->x, tile->header->y, tile->header->layer);
+	
+	// Create dtOffMeshSegmentIntersectionLinkPair list.
+	const int maxLinkPairs = DT_MAX_OFFMESH_SEGMENT_POINTS - 1;
+	dtOffMeshSegmentIntersectionLinkPair linkPairs[maxLinkPairs];
+	memset(linkPairs, 0, sizeof(dtOffMeshSegmentIntersectionLinkPair) * maxLinkPairs);
 
-	// match from A to B
-	for (int i = 0; i < segData.listA.npoints; i++)
-	{
-		unsigned int idxB = findMatchingSegmentIntersection(segData.listA.points[i].t, segData.listB.points, segData.listB.npoints, true);
-		if (idxB != DT_INVALID_SEGMENT)
-		{
-			links[nlinks].t = segData.listA.points[i].t;
-			links[nlinks].polyA = segData.listA.points[i].poly;
-			links[nlinks].polyB = segData.listB.points[idxB].poly;
-			links[nlinks].tileA = segData.listA.points[i].tile;
-			links[nlinks].tileB = segData.listB.points[idxB].tile;
-			nlinks++;
-		}
-	}
+	int nlinkPairs = findSegmentIntersectionLinkPairs(segData.listA, segData.listB, linkPairs);
 
-	// match from B to A
-	for (int i = 0; i < segData.listA.npoints; i++)
+	// There needs to be at least one pair of links to create a segment part.
+	if (nlinkPairs < 1)
 	{
-		unsigned int idxA = findMatchingSegmentIntersection(segData.listB.points[i].t, segData.listA.points, segData.listA.npoints, false);
-		if (idxA != DT_INVALID_SEGMENT)
-		{
-			links[nlinks].t = segData.listB.points[i].t;
-			links[nlinks].polyA = segData.listA.points[idxA].poly;
-			links[nlinks].polyB = segData.listB.points[i].poly;
-			links[nlinks].tileA = segData.listA.points[idxA].tile;
-			links[nlinks].tileB = segData.listB.points[i].tile;
-			nlinks++;
-		}
+		return;
 	}
 
-	if (nlinks < 2)
-		return;
+	UE_LOG(LogSegmentLink, Log, TEXT("   %hs %i nlinkPairs"), __FUNCTION__, nlinkPairs);
 
-	// sort positions
-	qsort(links, nlinks, sizeof(dtOffMeshSegmentIntersectionLink), segmentIntersectionLinkSorter);
+	qsort(linkPairs, nlinkPairs, sizeof(dtOffMeshSegmentIntersectionLinkPair), segmentIntersectionLinkPairSorter);
 
 	// create segments
 	memset(parts, 0, sizeof(dtOffMeshSegmentPart)*maxParts);
 	nparts = 0;
 
-	for (int i = 1; i < nlinks; i++)
+	// Create segment parts for link pairs that belong to the current tile.
+	for (int i = 0; i < nlinkPairs; i++)
 	{
-		if (links[i-1].tileA == tile || links[i-1].tileB == tile ||
-			links[i  ].tileA == tile || links[i  ].tileB == tile ||
-			canConnectSegmentPart(links[i-1].polyA, links[i].polyA, links[i-1].tileA, links[i].tileA, segData.listA.points, segData.listA.npoints))
+		const dtOffMeshSegmentIntersectionLinkPair& linkPair = linkPairs[i];
+		
+		if (linkPair.links[0].tileA == tile || linkPair.links[0].tileB == tile ||
+			linkPair.links[1].tileA == tile || linkPair.links[1].tileB == tile)
 		{
-			parts[nparts].t0 = links[i-1].t;
-			parts[nparts].t1 = links[i].t;
-			parts[nparts].polyA = links[i-1].polyA;
-			parts[nparts].polyB = links[i-1].polyB;
-			parts[nparts].tileA = links[i-1].tileA;
-			parts[nparts].tileB = links[i-1].tileB;
+			dtOffMeshSegmentPart& part = parts[nparts];
+
+			part.tA0 = linkPair.links[0].tA;
+			part.tA1 = linkPair.links[1].tA;
+			part.tB0 = linkPair.links[0].tB;
+			part.tB1 = linkPair.links[1].tB;
+			part.polyA = linkPair.links[0].polyA;
+			part.polyB = linkPair.links[0].polyB;
+			part.tileA = linkPair.links[0].tileA;
+			part.tileB = linkPair.links[0].tileB;
+
 			nparts++;
+			
+			UE_LOG(LogSegmentLink, Log, TEXT("   %hs seg part %i created tA: %f, %f, tb: %f, %f, tileA (%i,%i,%i), tileB (%i,%i,%i), tile=%p, tileA=%p, tileB=%p"),
+				__FUNCTION__, nparts,
+				part.tA0, part.tA1,
+				part.tB0, part.tB1,
+				part.tileA->header->x, part.tileA->header->y, part.tileA->header->layer,
+				part.tileB->header->x, part.tileB->header->y, part.tileB->header->layer,
+				tile, part.tileA, part.tileB);
 		}
 	}
 
@@ -662,17 +777,25 @@ static void createSegmentParts(dtMeshTile* tile, const dtOffMeshSegmentData& seg
 	unsigned short nPos = 0;
 	for (int i = 0; i < nparts; i++)
 	{
-		parts[i].vA0 = findOrAddUniqueValue(parts[i].t0, uniquePos, nPos);
-		parts[i].vA1 = findOrAddUniqueValue(parts[i].t1, uniquePos, nPos);
+		// Use specific t values for the A segment part.
+		parts[i].vA0 = findOrAddUniqueValue(parts[i].tA0, uniquePos, nPos);
+		parts[i].vA1 = findOrAddUniqueValue(parts[i].tA1, uniquePos, nPos);
 	}
+	
+	dtReal uniquePosB[DT_MAX_OFFMESH_SEGMENT_PARTS*2];
+	unsigned short nPosB = 0;
 
 	for (int i = 0; i < nparts; i++)
 	{
-		parts[i].vB0 = parts[i].vA0 + nPos;
-		parts[i].vB1 = parts[i].vA1 + nPos;
+		// Use specific t values for the B segment part.
+		// B values can also be unique.
+		parts[i].vB0 = findOrAddUniqueValue(parts[i].tB0, uniquePosB, nPosB) + nPos;
+		parts[i].vB1 = findOrAddUniqueValue(parts[i].tB1, uniquePosB, nPosB) + nPos;
 	}
 
-	nverts = nPos * 2;
+	nverts = (nPos + nPosB) * 2;
+	
+	UE_LOG(LogSegmentLink, Log, TEXT("   %hs %i nparts created"), __FUNCTION__, nparts);
 }
 
 static void createSegmentPolys(dtNavMesh* nav, dtMeshTile* tile, dtOffMeshSegmentConnection* con,
@@ -692,10 +815,10 @@ static void createSegmentPolys(dtNavMesh* nav, dtMeshTile* tile, dtOffMeshSegmen
 		dtOffMeshSegmentPart* it = &parts[i];
 
 		// add verts
-		dtVmad(&tile->verts[(vertBase+it->vA0)*3], con->startA, lenA, it->t0);
-		dtVmad(&tile->verts[(vertBase+it->vA1)*3], con->startA, lenA, it->t1);
-		dtVmad(&tile->verts[(vertBase+it->vB0)*3], con->startB, lenB, it->t0);
-		dtVmad(&tile->verts[(vertBase+it->vB1)*3], con->startB, lenB, it->t1);
+		dtVmad(&tile->verts[(vertBase+it->vA0)*3], con->startA, lenA, it->tA0);
+		dtVmad(&tile->verts[(vertBase+it->vA1)*3], con->startA, lenA, it->tA1);
+		dtVmad(&tile->verts[(vertBase+it->vB0)*3], con->startB, lenB, it->tB0);
+		dtVmad(&tile->verts[(vertBase+it->vB1)*3], con->startB, lenB, it->tB1);
 
 		// add poly
 		dtPoly* poly = &tile->polys[polyBase + i];
@@ -718,10 +841,12 @@ static void createSegmentPolys(dtNavMesh* nav, dtMeshTile* tile, dtOffMeshSegmen
 
 static void createSegmentLinks(dtNavMesh* nav, dtOffMeshSegmentData* seg, dtMeshTile* tile)
 {
-	if (seg == NULL)
+	if (seg == nullptr)
 		return;
 
-	unsigned short vertBase = (unsigned short)tile->header->offMeshSegVertBase;
+	UE_LOG(LogSegmentLink, Log, TEXT("%hs (%i,%i,%i)"), __FUNCTION__, tile->header->x, tile->header->y, tile->header->layer);
+
+	unsigned short vertBase = tile->header->offMeshSegVertBase;
 	int polyBase = tile->header->offMeshSegPolyBase;
 	for (int i = 0; i < tile->header->offMeshSegConCount; i++)
 	{
@@ -739,7 +864,11 @@ static void createSegmentLinks(dtNavMesh* nav, dtOffMeshSegmentData* seg, dtMesh
 
 		createSegmentPolys(nav, tile, &con, parts, nparts, vertBase, polyBase);
 		vertBase += nverts;
-		polyBase += nparts;
+
+		// The polys array is allocated enough memory to contain the max off-mesh segment parts when it is created.
+		// Therefore, polyBase must be incremented by the max off-mesh segment parts count for subsequent connections to access the correct data.
+		// See dtCreateNavMeshData for polys array creation and initialization.
+		polyBase += DT_MAX_OFFMESH_SEGMENT_PARTS;
 	}
 }
 #endif // WITH_NAVMESH_SEGMENT_LINKS
@@ -1124,7 +1253,7 @@ void dtNavMesh::connectExtLinks(dtMeshTile* tile, dtMeshTile* target, int side,
 }
 
 //@UE BEGIN
-void dtNavMesh::linkOffMeshHelper(dtMeshTile* tile0, unsigned int polyIdx0, dtMeshTile* tile1, unsigned int polyIdx1, unsigned char side, unsigned char edge)
+void dtNavMesh::linkOffMeshHelper(dtMeshTile* tile0, unsigned int polyIdx0, const dtMeshTile* tile1, unsigned int polyIdx1, unsigned char side, unsigned char edge)
 {
 	dtPoly* poly0 = &tile0->polys[polyIdx0];
 
@@ -1865,10 +1994,6 @@ dtStatus dtNavMesh::addTile(unsigned char* data, int dataSize, int flags,
 
 	connectIntLinks(tile);
 	baseOffMeshLinks(tile);
-	
-#if WITH_NAVMESH_SEGMENT_LINKS
-	dtOffMeshSegmentData* segList = initSegmentIntersection(*this, tile);
-#endif // WITH_NAVMESH_SEGMENT_LINKS
 
 	// Create connections with neighbour tiles.
 	ReadTilesHelper TileArray;
@@ -1890,9 +2015,6 @@ dtStatus dtNavMesh::addTile(unsigned char* data, int dataSize, int flags,
 		{
 			connectExtLinks(tile, neis[j], -1, bHasClusters);
 			connectExtLinks(neis[j], tile, -1, bHasClusters);
-#if WITH_NAVMESH_SEGMENT_LINKS
-			appendSegmentIntersection(*this, segList, tile, neis[j]);
-#endif // WITH_NAVMESH_SEGMENT_LINKS
 			connectExtOffMeshLinks(tile, neis[j], -1, bHasClusters);
 		}
 		connectExtOffMeshLinks(neis[j], tile, -1, bHasClusters);
@@ -1919,19 +2041,11 @@ dtStatus dtNavMesh::addTile(unsigned char* data, int dataSize, int flags,
 				connectExtLinks(tile, neis[j], i, bHasClusters);
 				connectExtLinks(neis[j], tile, dtOppositeTile(i), bHasClusters);
 			}
-#if WITH_NAVMESH_SEGMENT_LINKS
-			appendSegmentIntersection(*this, segList, tile, neis[j]);
-#endif // WITH_NAVMESH_SEGMENT_LINKS
 
 			connectExtOffMeshLinks(tile, neis[j], i, bHasClusters);
 			connectExtOffMeshLinks(neis[j], tile, dtOppositeTile(i), bHasClusters);
 		}
 	}
-
-#if WITH_NAVMESH_SEGMENT_LINKS
-	createSegmentLinks(this, segList, tile);
-	dtFree(segList, DT_ALLOC_TEMP);
-#endif // WITH_NAVMESH_SEGMENT_LINKS
 	
 	if (result)
 		*result = getTileRef(tile);
@@ -1942,6 +2056,98 @@ dtStatus dtNavMesh::addTile(unsigned char* data, int dataSize, int flags,
 }
 //@UE END
 
+#if WITH_NAVMESH_SEGMENT_LINKS
+void dtNavMesh::processSegmentLinksForTile(dtTileRef tileRef, unsigned int maxSkippedNeigborTiles, dtTileRef* skippedNeigborTiles, unsigned int& numSkippedNeighborTiles)
+{
+	dtMeshTile* tile = getMutableTileByRef(tileRef);
+
+	if (tile == nullptr || tile->header == nullptr)
+	{
+		return;
+	}
+
+	const int tileX = tile->header->x;
+	const int tileY = tile->header->y;
+	const unsigned short tileL = tile->header->layer;
+	
+	UE_LOG(LogSegmentLink, Log, TEXT("%hs (%i,%i,%i)"), __FUNCTION__, tile->header->x, tile->header->y, tile->header->layer);
+	
+	// Find list of intersections with tile polys.
+	dtOffMeshSegmentData* segList = initSegmentIntersection(*this, tile);
+
+	if (segList == nullptr)
+	{
+		return;
+	}
+
+	numSkippedNeighborTiles = 0;
+
+	// Create connections with neighbour tiles.
+	ReadTilesHelper TileArray;
+
+	// Connect with layers in current tile.
+	UE_LOG(LogSegmentLink, Log, TEXT("   %hs Connect with layers in current tile."), __FUNCTION__);
+	int nneis = getTileCountAt(tileX, tileY);
+	dtMeshTile** neis = TileArray.PrepareArray(nneis);
+	getTilesAt(tileX, tileY, neis, nneis);
+	for (int j = 0; j < nneis; ++j)
+	{
+		dtMeshTile* neiTile = neis[j];
+		if (neiTile != tile && neiTile->header)
+		{
+			const int neiX = neiTile->header->x;
+			const int neiY = neiTile->header->y;
+			const unsigned short neiL = neiTile->header->layer;
+ 			if (neiX < tileX && neiY < tileY && neiL < tileL)
+			{
+				if (numSkippedNeighborTiles < maxSkippedNeigborTiles)
+				{
+					skippedNeigborTiles[numSkippedNeighborTiles] = getTileRef(neis[j]);
+				}
+				++numSkippedNeighborTiles;
+				continue;
+			}
+
+			appendSegmentIntersection(*this, segList, tile, neis[j]);
+		}
+	}
+
+	// Connect with neighbour tiles.
+	UE_LOG(LogSegmentLink, Log, TEXT("   %hs Connect with neighbour tiles."), __FUNCTION__);
+	for (int i = 0; i < 8; ++i)
+	{
+		nneis = getNeighbourTilesCountAt(tileX, tileY, i);
+		neis = TileArray.PrepareArray(nneis);
+		getNeighbourTilesAt(tileX, tileY, i, neis, nneis);
+
+		for (int j = 0; j < nneis; ++j)
+		{
+			dtMeshTile* neiTile = neis[j];
+			if (neiTile->header)
+			{
+				const int neiX = neiTile->header->x;
+				const int neiY = neiTile->header->y;
+				const unsigned short neiL = neiTile->header->layer;
+				if (neiX < tileX && neiY < tileY && neiL < tileL)
+				{
+					if (numSkippedNeighborTiles < maxSkippedNeigborTiles)
+					{
+						skippedNeigborTiles[numSkippedNeighborTiles] = getTileRef(neis[j]);
+					}
+					++numSkippedNeighborTiles;
+					continue;
+				}
+
+				appendSegmentIntersection(*this, segList, tile, neis[j]);
+			}
+		}
+	}
+
+	createSegmentLinks(this, segList, tile);
+	dtFree(segList, DT_ALLOC_TEMP);
+}
+#endif // WITH_NAVMESH_SEGMENT_LINKS
+
 const dtMeshTile* dtNavMesh::getTileAt(const int x, const int y, const int layer) const
 {
 	// Find tile based on hash.
@@ -2102,6 +2308,20 @@ const dtMeshTile* dtNavMesh::getTileByRef(dtTileRef ref) const
 	return tile;
 }
 
+dtMeshTile* dtNavMesh::getMutableTileByRef(dtTileRef ref) const
+{
+	if (!ref)
+		return 0;
+	unsigned int tileIndex = decodePolyIdTile((dtPolyRef)ref);
+	unsigned int tileSalt = decodePolyIdSalt((dtPolyRef)ref);
+	if ((int)tileIndex >= m_maxTiles)
+		return 0;
+	dtMeshTile* tile = &m_tiles[tileIndex];
+	if (tile->salt != tileSalt)
+		return 0;
+	return tile;
+}
+
 int dtNavMesh::getMaxTiles() const
 {
 	return m_maxTiles;
@@ -2437,6 +2657,8 @@ dtStatus dtNavMesh::restoreTileState(dtMeshTile* tile, const unsigned char* data
 /// inside a normal polygon. So an off-mesh connection is "entered" from a 
 /// normal polygon at one of its endpoints. This is the polygon identified by 
 /// the prevRef parameter.
+/// For segment links, find the starting position closest to currentPos return the end position closest to that starting position.
+/// If currentPos is not provided, use the midpoint on the starting edge as the starting position.
 dtStatus dtNavMesh::getOffMeshConnectionPolyEndPoints(dtPolyRef prevRef, dtPolyRef polyRef, const dtReal* currentPos, dtReal* startPos, dtReal* endPos) const
 {
 	unsigned int salt, it, ip;
@@ -2489,10 +2711,27 @@ dtStatus dtNavMesh::getOffMeshConnectionPolyEndPoints(dtPolyRef prevRef, dtPolyR
 		dtVcopy(start0, &tile->verts[poly->verts[idx0]*3]);
 		dtVcopy(start1, &tile->verts[poly->verts[idx1]*3]);
 		dtReal t = 0;
-		dtDistancePtSegSqr2D(startPos, start0, start1, t);
+		
+		// Find the point on the segment link edge that is closest to the agent's current position.
+		// However, this function can be called with a null currentPos. In that case, just use the midpoint of the edge.
+		if (currentPos != nullptr)
+		{
+			dtDistancePtSegSqr2D(currentPos, start0, start1, t);
+		}
+		else
+		{
+			t = 0.5;
+		}
 
 		dtVlerp(startPos, start0, start1, t);
-		dtVlerp(endPos, &tile->verts[poly->verts[idx2]*3], &tile->verts[poly->verts[idx3]*3], t);
+		
+		// We want to use the shortest path to cross the segment link as this is generally the most sensible thing to do.
+		// This change also matches the change to the string pulling algorithm used by traditional path following.
+		dtReal end0[3], end1[3];
+		dtVcopy(end0, &tile->verts[poly->verts[idx2] * 3]);
+		dtVcopy(end1, &tile->verts[poly->verts[idx3] * 3]);
+		dtDistancePtSegSqr2D(startPos, end0, end1, t);
+		dtVlerp(endPos, end0, end1, t);
 	}
 	else
 #endif // WITH_NAVMESH_SEGMENT_LINKS
@@ -2730,4 +2969,4 @@ void dtNavMesh::applyAreaCostOrder(unsigned char* costOrder)
 {
 	memcpy(m_areaCostOrder, costOrder, sizeof(m_areaCostOrder));
 }
-//@UE END
+//@UE END
\ No newline at end of file
diff --git a/Engine/Source/Runtime/Navmesh/Private/Detour/DetourNavMeshBuilder.cpp b/Engine/Source/Runtime/Navmesh/Private/Detour/DetourNavMeshBuilder.cpp
--- a/Engine/Source/Runtime/Navmesh/Private/Detour/DetourNavMeshBuilder.cpp
+++ b/Engine/Source/Runtime/Navmesh/Private/Detour/DetourNavMeshBuilder.cpp
@@ -303,7 +303,7 @@ bool dtCreateNavMeshData(dtNavMeshCreateParams* params, unsigned char** outData,
 
 		memset(offMeshConClass, 0, sizeof(unsigned char)*params->offMeshConCount*2);
 
-		// Find tight heigh bounds, used for culling out off-mesh start locations.
+		// Find tight height bounds, used for culling out off-mesh start locations.
 		dtReal hmin = DT_REAL_MAX;
 		dtReal hmax = -DT_REAL_MAX;
 		
@@ -333,12 +333,15 @@ bool dtCreateNavMeshData(dtNavMeshCreateParams* params, unsigned char** outData,
 		dtVcopy(bmax, params->bmax);
 		bmin[1] = hmin;
 		bmax[1] = hmax;
-		
+
+#if WITH_NAVMESH_SEGMENT_LINKS
+		// This is to fix an issues observed with dtIntersectSegmentPoly2D. The poly verts need to be in clockwise order.
 		dtReal bverts[3*4];
-		bverts[ 0] = bmin[0]; bverts[ 2] = bmin[2];
-		bverts[ 3] = bmax[0]; bverts[ 5] = bmin[2];
-		bverts[ 6] = bmax[0]; bverts[ 8] = bmax[2];
-		bverts[ 9] = bmin[0]; bverts[11] = bmax[2];
+		bverts[0] = bmin[0]; bverts[2] = bmax[2];
+		bverts[3] = bmax[0]; bverts[5] = bmax[2];
+		bverts[6] = bmax[0]; bverts[8] = bmin[2];
+		bverts[9] = bmin[0]; bverts[11] = bmin[2];
+#endif // WITH_NAVMESH_SEGMENT_LINKS
 
 		for (int i = 0; i < params->offMeshConCount; ++i)
 		{
@@ -362,6 +365,7 @@ bool dtCreateNavMeshData(dtNavMeshCreateParams* params, unsigned char** outData,
 				if (offMeshConClass[i*2+0] == 0xff)
 					storedOffMeshConCount++;
 			}
+#if WITH_NAVMESH_SEGMENT_LINKS
 			else if (offMeshCon.type & DT_OFFMESH_CON_SEGMENT)
 			{
 				int smin, smax;
@@ -377,6 +381,7 @@ bool dtCreateNavMeshData(dtNavMeshCreateParams* params, unsigned char** outData,
 					storedOffMeshSegCount++;
 				}
 			}
+#endif // WITH_NAVMESH_SEGMENT_LINKS
 		}
 	}
 	
diff --git a/Engine/Source/Runtime/Navmesh/Private/Detour/DetourNavMeshQuery.cpp b/Engine/Source/Runtime/Navmesh/Private/Detour/DetourNavMeshQuery.cpp
--- a/Engine/Source/Runtime/Navmesh/Private/Detour/DetourNavMeshQuery.cpp
+++ b/Engine/Source/Runtime/Navmesh/Private/Detour/DetourNavMeshQuery.cpp
@@ -2506,6 +2506,15 @@ dtStatus dtNavMeshQuery::appendPortals(const int startIdx, const int endIdx, con
 	return DT_IN_PROGRESS;
 }
 
+#if WITH_NAVMESH_SEGMENT_LINKS
+enum class ESegmentLinkPortalEntryTestState : uint8
+{
+	None,
+	Requested,
+	Active
+};
+#endif // WITH_NAVMESH_SEGMENT_LINKS
+
 /// @par
 /// 
 /// This method peforms what is often called 'string pulling'.
@@ -2565,6 +2574,14 @@ dtStatus dtNavMeshQuery::findStraightPath(const dtReal* startPos, const dtReal*
 
 		dtPolyRef leftPolyRef = path[0];
 		dtPolyRef rightPolyRef = path[0];
+
+#if WITH_NAVMESH_SEGMENT_LINKS
+		// Store the entry point to test when entering a segment link portal.
+		dtReal segmentLinkEntryPoint[3];
+
+		// Store the current state of the segment link portal entry test.
+		ESegmentLinkPortalEntryTestState segmentLinkPortalEntryTestState = ESegmentLinkPortalEntryTestState::None;
+#endif // WITH_NAVMESH_SEGMENT_LINKS
 		
 		for (int i = 0; i < pathSize; ++i)
 		{
@@ -2618,8 +2635,9 @@ dtStatus dtNavMeshQuery::findStraightPath(const dtReal* startPos, const dtReal*
 			// Lock moving through segment off-mesh connections
 			if (fromType == DT_POLYTYPE_OFFMESH_SEGMENT)
 			{
-				if (segSwapped)
-					segt = 1.0f - segt;
+				// Use the shortest path to exit the segment link, so just find the point on the exit edge
+				// that is closest to the point used on the entrance edge.
+				dtDistancePtSegSqr2D(segmentLinkEntryPoint, left, right, segt);
 
 				dtReal lockedPortal[3];
 				dtVlerp(lockedPortal, left, right, segt);
@@ -2627,22 +2645,28 @@ dtStatus dtNavMeshQuery::findStraightPath(const dtReal* startPos, const dtReal*
 				dtVcopy(right, lockedPortal);
 			}
 
-			segSwapped = false;
-			if (toType == DT_POLYTYPE_OFFMESH_SEGMENT && i != apexIndex)
+			// The value of seqSwapped should not be calculated if this iteration is occurring due to a segment link portal entry test.
+			// In this case, the value of seqSwapped was previously calculated during the previous iteration.
+			if (segmentLinkPortalEntryTestState == ESegmentLinkPortalEntryTestState::None)
 			{
-				dtReal mid0[3], mid1[3];
-				dtVadd(mid0, portalLeft, portalRight);
-				dtVscale(mid0, mid0, 0.5f);
-				dtVadd(mid1, left, right);
-				dtVscale(mid1, mid1, 0.5f);
-				dtReal dirm[3], dir0[3], dir1[3];
-				dtVsub(dirm, mid1, mid0);
-				dtVsub(dir0, portalLeft, mid0);
-				dtVsub(dir1, left, mid1);
-				const dtReal c0 = dtVperp2D(dirm, dir0);
-				const dtReal c1 = dtVperp2D(dirm, dir1);
-				segSwapped = ((c0 > 0.f) && (c1 < 0.f)) || ((c0 < 0.f) && (c1 > 0.f));
+				segSwapped = false;
+				if (toType == DT_POLYTYPE_OFFMESH_SEGMENT && i != apexIndex)
+				{
+					dtReal mid0[3], mid1[3];
+					dtVadd(mid0, portalLeft, portalRight);
+					dtVscale(mid0, mid0, 0.5f);
+					dtVadd(mid1, left, right);
+					dtVscale(mid1, mid1, 0.5f);
+					dtReal dirm[3], dir0[3], dir1[3];
+					dtVsub(dirm, mid1, mid0);
+					dtVsub(dir0, portalLeft, mid0);
+					dtVsub(dir1, left, mid1);
+					const dtReal c0 = dtVperp2D(dirm, dir0);
+					const dtReal c1 = dtVperp2D(dirm, dir1);
+					segSwapped = ((c0 > 0.f) && (c1 < 0.f)) || ((c0 < 0.f) && (c1 > 0.f));
+				}
 			}
+
 			if (segSwapped)
 			{
 				dtReal tmp[3];
@@ -2650,6 +2674,22 @@ dtStatus dtNavMeshQuery::findStraightPath(const dtReal* startPos, const dtReal*
 				dtVcopy(left, right);
 				dtVcopy(right, tmp);
 			}
+
+			// Handle requested and active segment link portal entry tests.
+			if (toType == DT_POLYTYPE_OFFMESH_SEGMENT && segmentLinkPortalEntryTestState == ESegmentLinkPortalEntryTestState::Requested)
+			{
+				// Lock left and right to test the single segment link entry point.
+				dtVcopy(left, segmentLinkEntryPoint);
+				dtVcopy(right, segmentLinkEntryPoint);
+
+				// The segment link entry point check is now active.
+				segmentLinkPortalEntryTestState = ESegmentLinkPortalEntryTestState::Active;
+			}
+			else
+			{
+				// Segment link entry point check is not active for this iteration.
+				segmentLinkPortalEntryTestState = ESegmentLinkPortalEntryTestState::None;
+			}
 #endif // WITH_NAVMESH_SEGMENT_LINKS
 //@UE END
 
@@ -2695,7 +2735,8 @@ dtStatus dtNavMeshQuery::findStraightPath(const dtReal* startPos, const dtReal*
 
 					//@UE BEGIN
 #if WITH_NAVMESH_SEGMENT_LINKS
-					if (toType == DT_POLYTYPE_OFFMESH_SEGMENT)
+					// The value of segt has already been calculated in the previous iteration when a segment link portal entry check is active.
+					if (toType == DT_POLYTYPE_OFFMESH_SEGMENT && segmentLinkPortalEntryTestState == ESegmentLinkPortalEntryTestState::None)
 						dtDistancePtSegSqr2D(portalApex, left, right, segt);
 #endif // WITH_NAVMESH_SEGMENT_LINKS
 					//@UE END
@@ -2749,7 +2790,8 @@ dtStatus dtNavMeshQuery::findStraightPath(const dtReal* startPos, const dtReal*
 
 					//@UE BEGIN
 #if WITH_NAVMESH_SEGMENT_LINKS
-					if (toType == DT_POLYTYPE_OFFMESH_SEGMENT)
+					// The value of segt has already been calculated in the previous iteration when a segment link portal entry check is active.
+					if (toType == DT_POLYTYPE_OFFMESH_SEGMENT && segmentLinkPortalEntryTestState == ESegmentLinkPortalEntryTestState::None)
 						dtDistancePtSegSqr2D(portalApex, left, right, segt);
 #endif // WITH_NAVMESH_SEGMENT_LINKS
 					//@UE END
@@ -2766,8 +2808,37 @@ dtStatus dtNavMeshQuery::findStraightPath(const dtReal* startPos, const dtReal*
 			// Handle entering off-mesh segments
 			if (toType == DT_POLYTYPE_OFFMESH_SEGMENT)
 			{
-				dtDistancePtSegSqr2D(portalApex, left, right, segt);
-				dtVlerp(portalApex, left, right, segt);
+				// Check if the segment link entry point needs to be tested against the current portal apex.
+				if (segmentLinkPortalEntryTestState == ESegmentLinkPortalEntryTestState::None)
+				{
+					// Segment link portal entry test is not needed when the link is being entered directly from
+					// the current apex because there is no current funnel to test against. There is guaranteed
+					// to be a clear path to the segment link portal entry point.
+					if (i == apexIndex)
+					{
+						dtDistancePtSegSqr2D(portalApex, left, right, segt);
+						dtVlerp(portalApex, left, right, segt);
+
+						// segmentLinkEntryPoint still needs to be set to determine the exit point on the other edge of the segment link.
+						dtVcopy(segmentLinkEntryPoint, portalApex);
+					}
+					else
+					{
+						// Calculate the value of segt to use during the segment link portal entry point test.
+						dtDistancePtSegSqr2D(portalApex, left, right, segt);
+
+						// Set the test point for segment link entry instead of portalApex directly.
+						// This point will be used as the new portal apex if it passes the funnel algorithm.
+						dtVlerp(segmentLinkEntryPoint, left, right, segt);
+
+						// Set up the next iteration for the segment link portal test. 
+						segmentLinkPortalEntryTestState = ESegmentLinkPortalEntryTestState::Requested;
+
+						// The current iteration is restarted to perform the segment link portal test.
+						--i;
+						continue;
+					}
+				}
 
 				stat = appendVertex(portalApex, DT_STRAIGHTPATH_OFFMESH_CONNECTION, path[i + 1], result);
 				if (stat != DT_IN_PROGRESS)
@@ -5261,4 +5332,4 @@ bool dtNavMeshQuery::isInClosedList(dtPolyRef ref) const
 	if (!m_nodePool) return false;
 	const dtNode* node = m_nodePool->findNode(ref);
 	return node && node->flags & DT_NODE_CLOSED;
-}
+}
\ No newline at end of file
diff --git a/Engine/Source/Runtime/Navmesh/Public/Detour/DetourNavMesh.h b/Engine/Source/Runtime/Navmesh/Public/Detour/DetourNavMesh.h
--- a/Engine/Source/Runtime/Navmesh/Public/Detour/DetourNavMesh.h
+++ b/Engine/Source/Runtime/Navmesh/Public/Detour/DetourNavMesh.h
@@ -30,6 +30,10 @@
 
 NAVMESH_API DECLARE_LOG_CATEGORY_EXTERN(LogDetour, Log, All);
 
+#if WITH_NAVMESH_SEGMENT_LINKS
+NAVMESH_API DECLARE_LOG_CATEGORY_EXTERN(LogSegmentLink, Warning, All);
+#endif
+
 // Note: If you want to use 64-bit refs, change the types of both dtPolyRef & dtTileRef.
 // It is also recommended that you change dtHashRef() to a proper 64-bit hash.
 
@@ -266,8 +270,8 @@ struct dtBVNode
 };
 
 //@UE BEGIN
-// This is an unsupported feature and has not been finished to production quality.
 #if WITH_NAVMESH_SEGMENT_LINKS
+/// This is an experimental feature and has not been finished to production quality.
 struct dtOffMeshSegmentConnection
 {
 	dtReal startA[3];				///< Start point of segment A
@@ -391,7 +395,7 @@ struct dtMeshHeader
 	unsigned short detailTriCount;		///< The number of triangles in the detail mesh.
 	unsigned short bvNodeCount;			///< The number of bounding volume nodes. (Zero if bounding volumes are disabled.)
 	unsigned short offMeshConCount;		///< The number of point type off-mesh connections.
-	unsigned short offMeshBase;			///< The index of the first polygon which is an point type off-mesh connection.
+	unsigned short offMeshBase;			///< The index of the first polygon which is a point type off-mesh connection.
 
 	//@UE BEGIN
 #if WITH_NAVMESH_SEGMENT_LINKS
@@ -528,6 +532,15 @@ public:
 	///  @param[out]	result		The tile reference. (If the tile was succesfully added.) [opt]
 	/// @return The status flags for the operation.
 	NAVMESH_API dtStatus addTile(unsigned char* data, int dataSize, int flags, dtTileRef lastRef, dtTileRef* result);
+
+#if WITH_NAVMESH_SEGMENT_LINKS
+	/// Create segment link connections within a tile.
+	///  @param[in]		tileRef						The tileRef of the tile owning the segment link connections.
+	///  @param[in]		maxSkippedNeighborTiles		The max number of neighbor indices that can be held in skippedNeighborTiles.
+	///  @param[out]	skippedNeighborTiles		Array containing the indices of the neighbor tiles processed.
+	///  @param[out]	numSkippedNeighborTiles		The number of neighbor tiles skipped during link connection.
+	NAVMESH_API void processSegmentLinksForTile(dtTileRef tileRef, unsigned int maxSkippedNeighborTiles, dtTileRef* skippedNeighborTiles, unsigned int& numSkippedNeighborTiles);
+#endif // WITH_NAVMESH_SEGMENT_LINKS
 	
 	/// Removes the specified tile from the navigation mesh.
 	///  @param[in]		ref			The reference of the tile to remove.
@@ -594,6 +607,12 @@ public:
 	/// @return The tile for the specified reference, or null if the 
 	///		reference is invalid.
 	NAVMESH_API const dtMeshTile* getTileByRef(dtTileRef ref) const;
+
+	/// Gets the tile for the specified tile reference.
+	///  @param[in]	ref		The tile reference of the tile to retrieve.
+	/// @return The tile for the specified reference, or null if the 
+	///		reference is invalid.
+	NAVMESH_API dtMeshTile* getMutableTileByRef(dtTileRef ref) const;
 	
 	/// The maximum number of tiles supported by the navigation mesh.
 	/// @return The maximum number of tiles supported by the navigation mesh.
@@ -824,7 +843,7 @@ public:
 #endif // WITH_NAVMESH_CLUSTER_LINKS
 
 	/// Helper for creating links in off-mesh connections
-	NAVMESH_API void linkOffMeshHelper(dtMeshTile* tile0, unsigned int polyIdx0, dtMeshTile* tile1, unsigned int polyIdx1, unsigned char side, unsigned char edge);
+	NAVMESH_API void linkOffMeshHelper(dtMeshTile* tile0, unsigned int polyIdx0, const dtMeshTile* tile1, unsigned int polyIdx1, unsigned char side, unsigned char edge);
 
 	inline bool isEmpty() const
 	{
