Fix static analysis warning
#rnx

diff --git a/Engine/Source/Runtime/Navmesh/Private/Detour/DetourNavMeshQuery.cpp b/Engine/Source/Runtime/Navmesh/Private/Detour/DetourNavMeshQuery.cpp
--- a/Engine/Source/Runtime/Navmesh/Private/Detour/DetourNavMeshQuery.cpp
+++ b/Engine/Source/Runtime/Navmesh/Private/Detour/DetourNavMeshQuery.cpp
@@ -2577,7 +2577,7 @@ dtStatus dtNavMeshQuery::findStraightPath(const dtReal* startPos, const dtReal*
 
 #if WITH_NAVMESH_SEGMENT_LINKS
 		// Store the entry point to test when entering a segment link portal.
-		dtReal segmentLinkEntryPoint[3];
+		dtReal segmentLinkEntryPoint[3] = { 0., 0., 0. };
 
 		// Store the current state of the segment link portal entry test.
 		ESegmentLinkPortalEntryTestState segmentLinkPortalEntryTestState = ESegmentLinkPortalEntryTestState::None;
