[Navlink] Fix navlink trajectory issue occurring when using large values for end heights tolerance.

diff --git a/Engine/Source/Runtime/Navmesh/Private/DebugUtils/DetourNavLinkDebugDraw.cpp b/Engine/Source/Runtime/Navmesh/Private/DebugUtils/DetourNavLinkDebugDraw.cpp
--- a/Engine/Source/Runtime/Navmesh/Private/DebugUtils/DetourNavLinkDebugDraw.cpp
+++ b/Engine/Source/Runtime/Navmesh/Private/DebugUtils/DetourNavLinkDebugDraw.cpp
@@ -104,11 +104,16 @@ namespace UE::Detour::Private
 		}
 	}
 
-	void drawTrajectorySlice(duDebugDraw* dd, dtReal* pa, dtReal* pb, const dtNavLinkBuilder::Trajectory2D* trajectory, const dtReal* trajectoryDir, const unsigned int color)
+	void drawTrajectorySlice(duDebugDraw* dd, const dtReal* pa, const dtReal* pb, const dtNavLinkBuilder::Trajectory2D* trajectory, const dtReal* trajectoryDir, const unsigned int color)
 	{
+		dtReal start[3];
+		dtReal end[3];
+		dtVcopy(start, pa);
+		dtVcopy(end, pb);
+		
 		// Offset start and end points to account for the agent radius.
-		dtVmad(pa, pa, trajectoryDir, -trajectory->radiusOverflow);
-		dtVmad(pb, pb, trajectoryDir,  trajectory->radiusOverflow);
+		dtVmad(start, pa, trajectoryDir, -trajectory->radiusOverflow);
+		dtVmad(end, pb, trajectoryDir,  trajectory->radiusOverflow);
 		
 		unsigned int colt = duTransCol(color, 50);
 		unsigned int colb = duTransCol(duLerpCol(color,duColor::black,96), 50);
@@ -120,7 +125,7 @@ namespace UE::Detour::Private
 		{
 			const dtNavLinkBuilder::TrajectorySample& s = trajectory->samples[i];
 			const dtReal u = (float)i / (float)(nsamples-1);
-			dtVlerp(p1, pa, pb, u);
+			dtVlerp(p1, start, end, u);
 			dtVcopy(p0, p1);
 			p0[1] += s.ymin;
 			p1[1] += s.ymax;
@@ -151,7 +156,7 @@ namespace UE::Detour::Private
 		{
 			const dtNavLinkBuilder::TrajectorySample& s = trajectory->samples[i];
 			const float u = (float)i / (float)(nsamples-1);
-			dtVlerp(p1, pa, pb, u);
+			dtVlerp(p1, start, end, u);
 			dtVcopy(p0, p1);
 			p0[1] += s.ymin;
 			p1[1] += s.ymax;
@@ -556,9 +561,6 @@ void duDebugDrawNavLinkBuilder(duDebugDraw* dd, const dtNavLinkBuilder& linkBuil
 						dtReal spt[3], ept[3];
 						dtVlerp(spt, es->start.p, es->start.q, u);
 						dtVlerp(ept, es->end.p, es->end.q, u);
-
-						spt[1] = ssmp->height;
-						ept[1] = esmp->height;
 						
 						if (ssmp->flags & dtNavLinkBuilder::UNRESTRICTED)
 							UE::Detour::Private::drawTrajectorySlice(dd, spt, ept, &es->trajectory, es->az, duColor::green);
diff --git a/Engine/Source/Runtime/Navmesh/Private/Detour/DetourNavLinkBuilder.cpp b/Engine/Source/Runtime/Navmesh/Private/Detour/DetourNavLinkBuilder.cpp
--- a/Engine/Source/Runtime/Navmesh/Private/Detour/DetourNavLinkBuilder.cpp
+++ b/Engine/Source/Runtime/Navmesh/Private/Detour/DetourNavLinkBuilder.cpp
@@ -527,11 +527,16 @@ bool dtNavLinkBuilder::checkHeightfieldCollision(const dtReal x, const dtReal ym
 }
 
 // Returns true if none of the samples ymin, ymax collide with the heghtfield.
-bool dtNavLinkBuilder::isTrajectoryClear(dtReal* pa, dtReal* pb, const Trajectory2D* trajectory, const dtReal* trajectoryDir) const
+bool dtNavLinkBuilder::isTrajectoryClear(const dtReal* pa, const dtReal* pb, const Trajectory2D* trajectory, const dtReal* trajectoryDir) const
 {
+	dtReal start[3];
+	dtReal end[3];
+	dtVcopy(start, pa);
+	dtVcopy(end, pb);
+	
 	// Offset start and end points to account for the agent radius.
-	dtVmad(pa, pa, trajectoryDir, -trajectory->radiusOverflow);
-	dtVmad(pb, pb, trajectoryDir,  trajectory->radiusOverflow);
+	dtVmad(start, pa, trajectoryDir, -trajectory->radiusOverflow);
+	dtVmad(end, pb, trajectoryDir,  trajectory->radiusOverflow);
 
 	const int nsamples = trajectory->samples.Num();
 	const float invLastSample = 1.f / (nsamples-1);
@@ -540,7 +545,8 @@ bool dtNavLinkBuilder::isTrajectoryClear(dtReal* pa, dtReal* pb, const Trajector
 		dtReal p[3];
 		const TrajectorySample& s = trajectory->samples[i];
 		const float u = (float)i * invLastSample;
-		dtVlerp(p, pa, pb, u);
+		dtVlerp(p, start, end, u);
+
 		if (checkHeightfieldCollision(p[0], p[1] + s.ymin, p[1] + s.ymax, p[2]))
 		{
 			return false;
@@ -550,6 +556,7 @@ bool dtNavLinkBuilder::isTrajectoryClear(dtReal* pa, dtReal* pb, const Trajector
 	return true;	
 }
 
+// Add ground samples and set height on them.
 void dtNavLinkBuilder::sampleGroundSegment(GroundSegment* seg, const int nsamples, const float groundRange) const
 {
 	TRACE_CPUPROFILER_EVENT_SCOPE(dtNavLinkBuilder::sampleGroundSegment);
@@ -579,6 +586,58 @@ void dtNavLinkBuilder::sampleGroundSegment(GroundSegment* seg, const int nsample
 	}
 }
 
+void dtNavLinkBuilder::updateTrajectorySamples(EdgeSampler* es) const
+{
+	TRACE_CPUPROFILER_EVENT_SCOPE(dtNavLinkBuilder::updateTrajectorySamples);
+	
+	if (es->start.ngsamples != es->end.ngsamples)
+		return;
+	
+	const int nsamples = es->start.ngsamples;
+	
+	for (int i = 0; i < nsamples; ++i)
+	{
+		GroundSample& ssmp = es->start.gsamples[i];
+		GroundSample& esmp = es->end.gsamples[i];
+
+		// If there is no ground, the ground height will not be set.
+		if ((ssmp.flags & HAS_GROUND) == 0 || (esmp.flags & HAS_GROUND) == 0)
+			continue;
+
+		// When we sample ground segments, in sampleEdges, we have add least 2 samples.
+		check(nsamples >= 2);																
+		const dtReal u = (dtReal)i/(dtReal)(nsamples-1);
+		dtReal spt[3], ept[3];
+		dtVlerp(spt, es->start.p, es->start.q, u);
+		dtVlerp(ept, es->end.p, es->end.q, u);
+	
+		// Offset start and end points to account for the agent radius.
+		dtVmad(spt, spt, es->az, -es->trajectory.radiusOverflow);
+		dtVmad(ept, ept, es->az,  es->trajectory.radiusOverflow);
+
+		const int nTrajectorySamples = es->trajectory.samples.Num();
+		// When we initialize trajectory samples (initTrajectorySamples), we add at least 2 trajectory samples.
+		check(nTrajectorySamples >= 2);												
+		const float invLastTrajSample = 1.f / (nTrajectorySamples-1);
+		for (int trajIndex = 0; trajIndex < nTrajectorySamples; ++trajIndex)
+		{
+			dtReal p[3];
+			TrajectorySample& s = es->trajectory.samples[trajIndex];
+			const float trajU = (float)trajIndex * invLastTrajSample;
+			dtVlerp(p, spt, ept, trajU);
+
+			if (s.floorStart)
+			{
+				s.ymin = (ssmp.height + m_linkBuilderConfig.agentClimb) - p[1];		// -p[1] to stay relative to p[1]
+			}
+			else if (s.floorEnd)
+			{
+				s.ymin = (esmp.height + m_linkBuilderConfig.agentClimb) - p[1];		// -p[1] to stay relative to p[1]
+			}
+		}
+	}
+}
+
 void dtNavLinkBuilder::sampleAction(EdgeSampler* es) const
 {
 	TRACE_CPUPROFILER_EVENT_SCOPE(dtNavLinkBuilder::sampleAction);
@@ -601,9 +660,6 @@ void dtNavLinkBuilder::sampleAction(EdgeSampler* es) const
 		dtVlerp(spt, es->start.p, es->start.q, u);
 		dtVlerp(ept, es->end.p, es->end.q, u);
 		
-		spt[1] = ssmp.height;
-		ept[1] = esmp.height;
-		
 		if (!isTrajectoryClear(spt, ept, &es->trajectory, es->az))
 			continue;
 
@@ -611,7 +667,7 @@ void dtNavLinkBuilder::sampleAction(EdgeSampler* es) const
 	}
 }
 
-void dtNavLinkBuilder::initTrajectory(Trajectory2D* trajectory) const
+void dtNavLinkBuilder::initTrajectorySamples(Trajectory2D* trajectory) const
 {
 	using namespace UE::Detour::NavLink::Private;
 	
@@ -631,6 +687,9 @@ void dtNavLinkBuilder::initTrajectory(Trajectory2D* trajectory) const
 	const int nsamples  = dtMax(2, (int)ceilf(dx*m_invCs));
 	trajectory->samples.Reserve(nsamples);
 
+	const float dxSample = dx/nsamples;
+	const float roundedAgentRadius = dxSample > 0.f ? ceilf(agentRadius/dxSample)*dxSample : 0.f; 
+	
 	const float* spine = trajectory->spine;
 	unsigned char nspine = trajectory->nspine;
 
@@ -650,6 +709,13 @@ void dtNavLinkBuilder::initTrajectory(Trajectory2D* trajectory) const
 		TrajectorySample& s = trajectory->samples.Emplace_GetRef();
 		s.ymin = dtMin(dtMin(y0,y1), y2) + m_linkBuilderConfig.agentClimb - yRef;	
 		s.ymax = dtMax(dtMax(y0,y1), y2) + m_linkBuilderConfig.agentHeight - yRef; 
+
+		// Mark samples that need to be floored. 
+		// Mainly important for the end where the ground could be far from trajectory end point.
+		if (xRef >= (spine[0]-roundedAgentRadius) && xRef <= spine[0]+roundedAgentRadius)
+			s.floorStart = true;
+		else if (xRef >= (spine[(nspine-1)*2]-roundedAgentRadius) && xRef <= (spine[(nspine-1)*2]+roundedAgentRadius))
+			s.floorEnd = true;
 	}
 }
 
@@ -955,9 +1021,9 @@ bool dtNavLinkBuilder::sampleEdge(const dtLinkBuilderConfig& builderConfig, dtNa
 		samplingSeparationFactor = config.samplingSeparationFactor;
 		initJumpOverRig(es, &segs[ibest*6+0], &segs[ibest*6+3], -jumpStartDist, jumpStartDist, jumpHeight, groundRange);
 	}
-	
-	initTrajectory(&es->trajectory);
 
+	initTrajectorySamples(&es->trajectory);
+	
 	// Init start end segments.
 	dtReal offset[3];
 	trans2d(offset, es->az, es->ay, &es->trajectory.spine[0]);
@@ -976,6 +1042,9 @@ bool dtNavLinkBuilder::sampleEdge(const dtLinkBuilderConfig& builderConfig, dtNa
 	sampleGroundSegment(&es->start, ngsamples, es->groundRange);
 	sampleGroundSegment(&es->end, ngsamples, es->groundRange);
 
+	// Now that we have ground heights, update the trajectory samples.
+	updateTrajectorySamples(es);
+	
 	sampleAction(es);
 	
 	return true;
diff --git a/Engine/Source/Runtime/Navmesh/Public/Detour/DetourNavLinkBuilder.h b/Engine/Source/Runtime/Navmesh/Public/Detour/DetourNavLinkBuilder.h
--- a/Engine/Source/Runtime/Navmesh/Public/Detour/DetourNavLinkBuilder.h
+++ b/Engine/Source/Runtime/Navmesh/Public/Detour/DetourNavLinkBuilder.h
@@ -61,8 +61,11 @@ public:
 	struct TrajectorySample
 	{
 		// Min and max heights relative to the height on the segment between the start and the end point (at the sampling locations). 
-		float ymin;
-		float ymax;
+		float ymin = 0.f;
+		float ymax = 0.f;
+		
+		bool floorStart = false;
+		bool floorEnd = false;
 	};
 	
 	struct Trajectory2D
@@ -161,8 +164,8 @@ public:
 	NAVMESH_API void debugBuildEdge(const dtLinkBuilderConfig& acfg, dtNavLinkAction action, int edgeIndex, EdgeSampler& sampler);
 	
 private:
-	void initTrajectory(Trajectory2D* trajectory) const;
-	bool isTrajectoryClear(dtReal* pa, dtReal* pb, const Trajectory2D* trajectory, const dtReal* trajectoryDir) const;
+	void initTrajectorySamples(Trajectory2D* trajectory) const;
+	bool isTrajectoryClear(const dtReal* pa, const dtReal* pb, const Trajectory2D* trajectory, const dtReal* trajectoryDir) const;
 	
 	int findPotentialJumpOverEdges(const dtReal* sp, const dtReal* sq,
 								   const float depthRange, const float heightRange,
@@ -179,6 +182,10 @@ private:
 	bool checkHeightfieldCollision(const dtReal x, const dtReal ymin, const dtReal ymax, const dtReal z) const;
 
 	void sampleGroundSegment(GroundSegment* seg, const int nsamples, const float groundRange) const;
+
+	// Update the min height of the trajectory samples from the height of the ground at the start and the end of the trajectories.
+	// This method must be called after the segments have been sampled for ground. 
+	void updateTrajectorySamples(EdgeSampler* es) const;
 	
 	void sampleAction(EdgeSampler* es) const;
 	
