[Navlink Generation] Addition of "Selected Edge Collisions Samples" flag to display individual trajectory samples that hit a collisions when debugging tile generation.

diff --git a/Engine/Source/Runtime/NavigationSystem/Public/NavMesh/LinkGenerationDebugFlags.h b/Engine/Source/Runtime/NavigationSystem/Public/NavMesh/LinkGenerationDebugFlags.h
--- a/Engine/Source/Runtime/NavigationSystem/Public/NavMesh/LinkGenerationDebugFlags.h
+++ b/Engine/Source/Runtime/NavigationSystem/Public/NavMesh/LinkGenerationDebugFlags.h
@@ -6,13 +6,14 @@
 UENUM(meta = (Bitflags, UseEnumValuesAsMaskValuesInEditor = "true"))
 enum class ELinkGenerationDebugFlags : uint16
 {
-	WalkableSurface				= 1 << 0,
-	WalkableBorders				= 1 << 1,
-	SelectedEdge				= 1 << 2,
-	SelectedEdgeTrajectory		= 1 << 3,
-	SelectedEdgeLandingSamples	= 1 << 4,
-	SelectedEdgeCollisions		= 1 << 5,
-	Links						= 1 << 6,
-	FilteredLinks				= 1 << 7,
+	WalkableSurface				    = 1 << 0,
+	WalkableBorders				    = 1 << 1,
+	SelectedEdge				    = 1 << 2,
+	SelectedEdgeTrajectory		    = 1 << 3,
+	SelectedEdgeLandingSamples	    = 1 << 4,
+	SelectedEdgeCollisions		    = 1 << 5,
+	SelectedEdgeCollisionsSamples	= 1 << 6,
+	Links						    = 1 << 7,
+	FilteredLinks				    = 1 << 8,
 };
 ENUM_CLASS_FLAGS(ELinkGenerationDebugFlags);
diff --git a/Engine/Source/Runtime/Navmesh/Private/DebugUtils/DetourNavLinkDebugDraw.cpp b/Engine/Source/Runtime/Navmesh/Private/DebugUtils/DetourNavLinkDebugDraw.cpp
--- a/Engine/Source/Runtime/Navmesh/Private/DebugUtils/DetourNavLinkDebugDraw.cpp
+++ b/Engine/Source/Runtime/Navmesh/Private/DebugUtils/DetourNavLinkDebugDraw.cpp
@@ -181,6 +181,53 @@ namespace UE::Detour::Private
 	}
 } // UE::Detour::Private
 
+void duDebugDrawTrajectorySamples(duDebugDraw* dd, const dtNavLinkBuilder& linkBuilder, const dtReal* pa, const dtReal* pb,
+	const dtNavLinkBuilder::Trajectory2D* trajectory, const dtReal* trajectoryDir)
+{
+	dtReal start[3];
+	dtReal end[3];
+	dtVcopy(start, pa);
+	dtVcopy(end, pb);
+	
+	// Offset start and end points to account for the agent radius.
+	dtVmad(start, pa, trajectoryDir, -trajectory->radiusOverflow);
+	dtVmad(end, pb, trajectoryDir,  trajectory->radiusOverflow);
+
+	const int nsamples = trajectory->samples.Num();
+	const float invLastSample = 1.f / (nsamples-1);
+	for (int i = 0; i < nsamples; ++i)
+	{
+		dtReal p[3];
+		const dtNavLinkBuilder::TrajectorySample& s = trajectory->samples[i];
+		const float u = (float)i * invLastSample;
+		dtVlerp(p, start, end, u);
+
+		// Draw an additional point if the sample was marked to be snap to the floor.
+		if (s.floorStart)
+		{
+			dd->begin(DU_DRAW_POINTS, 5.0f);
+			dd->vertex(p[0], p[1]+s.ymin, p[2], duColor::black);
+			dd->end();
+		}
+		else if (s.floorEnd)
+		{
+			dd->begin(DU_DRAW_POINTS, 5.0f);
+			dd->vertex(p[0], p[1]+s.ymin, p[2], duColor::cyan);
+			dd->end();
+		}
+
+		// Check with the heightfield if a collision is hit.
+		const bool hit = linkBuilder.checkHeightfieldCollision(p[0], p[1] + s.ymin, p[1] + s.ymax, p[2]);
+		unsigned int color = hit ? duColor::orangeRed : duColor::lightGreen;
+
+		// Draw a vertical line for the sample.
+		dd->begin(DU_DRAW_LINES, 5.0f);
+		dd->vertex(p[0], p[1]+s.ymin, p[2], color);
+		dd->vertex(p[0], p[1]+s.ymax, p[2], color);
+		dd->end();
+	}
+}
+
 void duDebugDrawNavLinkBuilder(duDebugDraw* dd, const dtNavLinkBuilder& linkBuilder, unsigned int drawFlags, const dtNavLinkBuilder::EdgeSampler* es)
 {
 	if (dd == nullptr)
@@ -543,13 +590,14 @@ void duDebugDrawNavLinkBuilder(duDebugDraw* dd, const dtNavLinkBuilder& linkBuil
 					dd->end();
 				}
 			}
-			
-			if (drawFlags & DRAW_COLLISION_SLICES)
+
+			const bool drawCollisionSlices = drawFlags & DRAW_COLLISION_SLICES;
+			const bool drawCollisionSamples = drawFlags & DRAW_COLLISION_SAMPLES;
+			if (drawCollisionSlices || drawCollisionSamples)
 			{
 				if (es->start.ngsamples && es->start.ngsamples == es->end.ngsamples)
 				{
 					const int nsamples = es->start.ngsamples;
-					dd->depthMask(true);
 					for (int i = 0; i < nsamples; ++i)
 					{
 						const dtNavLinkBuilder::GroundSample* ssmp = &es->start.gsamples[i];
@@ -561,13 +609,20 @@ void duDebugDrawNavLinkBuilder(duDebugDraw* dd, const dtNavLinkBuilder& linkBuil
 						dtReal spt[3], ept[3];
 						dtVlerp(spt, es->start.p, es->start.q, u);
 						dtVlerp(ept, es->end.p, es->end.q, u);
-						
-						if (ssmp->flags & dtNavLinkBuilder::UNRESTRICTED)
-							UE::Detour::Private::drawTrajectorySlice(dd, spt, ept, &es->trajectory, es->az, duColor::green);
-						else
-							UE::Detour::Private::drawTrajectorySlice(dd, spt, ept, &es->trajectory, es->az, duColor::orangeRed);
+
+						if (drawCollisionSlices)
+						{
+							if (ssmp->flags & dtNavLinkBuilder::UNRESTRICTED)
+								UE::Detour::Private::drawTrajectorySlice(dd, spt, ept, &es->trajectory, es->az, duColor::green);
+							else
+								UE::Detour::Private::drawTrajectorySlice(dd, spt, ept, &es->trajectory, es->az, duColor::orangeRed);
+						}
+
+						if (drawCollisionSamples)
+						{
+							duDebugDrawTrajectorySamples(dd, linkBuilder, spt, ept, &es->trajectory, es->az);
+						}
 					}
-					dd->depthMask(true);
 				}
 			}
 		}
diff --git a/Engine/Source/Runtime/Navmesh/Public/DebugUtils/DetourNavLinkDebugDraw.h b/Engine/Source/Runtime/Navmesh/Public/DebugUtils/DetourNavLinkDebugDraw.h
--- a/Engine/Source/Runtime/Navmesh/Public/DebugUtils/DetourNavLinkDebugDraw.h
+++ b/Engine/Source/Runtime/Navmesh/Public/DebugUtils/DetourNavLinkDebugDraw.h
@@ -10,15 +10,19 @@ class dtNavLinkBuilder;
 
 enum duNavLinkBuilderDrawFlags
 {
-	DRAW_WALKABLE_SURFACE =	1 << 0,
-	DRAW_BORDERS =			1 << 1,
-	DRAW_SELECTED_EDGE =	1 << 2,
-	DRAW_TRAJECTORY =		1 << 3,
-	DRAW_LAND_SAMPLES =		1 << 4,
-	DRAW_COLLISION_SLICES =	1 << 5,
-	DRAW_LINKS =			1 << 6,
-	DRAW_FILTERED_LINKS =	1 << 7,
+	DRAW_WALKABLE_SURFACE =	    1 << 0,
+	DRAW_BORDERS =			    1 << 1,
+	DRAW_SELECTED_EDGE =	    1 << 2,
+	DRAW_TRAJECTORY =		    1 << 3,
+	DRAW_LAND_SAMPLES =		    1 << 4,
+	DRAW_COLLISION_SLICES =		1 << 5,
+	DRAW_COLLISION_SAMPLES =	1 << 6,
+	DRAW_LINKS =			    1 << 7,
+	DRAW_FILTERED_LINKS =	    1 << 8,
 };
 
 NAVMESH_API void duDebugDrawNavLinkBuilder(duDebugDraw* dd, const dtNavLinkBuilder& linkBuilder, unsigned int drawFlags, const struct dtNavLinkBuilder::EdgeSampler* es);
+
+NAVMESH_API void duDebugDrawTrajectorySamples(duDebugDraw* dd, const dtNavLinkBuilder& linkBuilder, const dtReal* pa, const dtReal* pb,
+											  const dtNavLinkBuilder::Trajectory2D* trajectory, const dtReal* trajectoryDir);
 //@UE END
diff --git a/Engine/Source/Runtime/Navmesh/Public/Detour/DetourNavLinkBuilder.h b/Engine/Source/Runtime/Navmesh/Public/Detour/DetourNavLinkBuilder.h
--- a/Engine/Source/Runtime/Navmesh/Public/Detour/DetourNavLinkBuilder.h
+++ b/Engine/Source/Runtime/Navmesh/Public/Detour/DetourNavLinkBuilder.h
@@ -152,6 +152,9 @@ private:
 
 	friend NAVMESH_API void duDebugDrawNavLinkBuilder(struct duDebugDraw* dd, const dtNavLinkBuilder& linkBuilder, unsigned int drawFlags, const EdgeSampler* es);
 	
+	friend NAVMESH_API void duDebugDrawTrajectorySamples(struct duDebugDraw* dd, const dtNavLinkBuilder& linkBuilder, const dtReal* pa, const dtReal* pb,
+														 const dtNavLinkBuilder::Trajectory2D* trajectory, const dtReal* trajectoryDir);
+	
 public:
 	// Loops through contours to store edge points in world coordinates.
 	NAVMESH_API bool findEdges(rcContext& ctx, const rcConfig& cfg, const dtLinkBuilderConfig& builderConfig,
