fix approximation error that could result in wrong span index being used
change rcAssert to an if for better resilience to user error

diff --git a/Engine/Source/Runtime/Navmesh/Private/Recast/RecastRasterization.cpp b/Engine/Source/Runtime/Navmesh/Private/Recast/RecastRasterization.cpp
--- a/Engine/Source/Runtime/Navmesh/Private/Recast/RecastRasterization.cpp
+++ b/Engine/Source/Runtime/Navmesh/Private/Recast/RecastRasterization.cpp
@@ -1056,9 +1056,10 @@ struct AddSpanInHeightfield
 
 struct AddSpanInTempColumns
 {
-	AddSpanInTempColumns(rcContext* ctx, rcHeightfield& hf, const float ics, const rcReal* vertsBMin, const rcReal* vertsBMax)
+	AddSpanInTempColumns(rcContext* inctx, rcHeightfield& hf, const float ics, const rcReal* vertsBMin, const rcReal* vertsBMax)
 	{
 #if EPIC_ADDITION_USE_NEW_RECAST_RASTERIZER
+		ctx = inctx;
 		if (vertsBMin == nullptr || vertsBMax == nullptr)
 		{
 			if (ctx)
@@ -1077,7 +1078,7 @@ struct AddSpanInTempColumns
 			return;
 		}
 
-		// init data
+		// clamp the verts bounding box to the heightfield
 		const rcReal xMin = rcClamp(vertsBMin[0], hf.bmin[0], hf.bmax[0]);
 		const rcReal xMax = rcClamp(vertsBMax[0], hf.bmin[0], hf.bmax[0]);
 		const rcReal zMin = rcClamp(vertsBMin[2], hf.bmin[2], hf.bmax[2]);
@@ -1088,11 +1089,17 @@ struct AddSpanInTempColumns
 			return;
 		}
 
-		// compute the part of the heightfield that will be covered by verts
-		ibminx = ics * (xMin - hf.bmin[0]);
-		ibminz = ics * (zMin - hf.bmin[2]);
-		ibwidth = rcClamp(ics * (xMax - xMin), 0, hf.width);
-		ibheight = rcClamp(ics * (zMax - zMin), 0, hf.height);
+		// convert it into heightfield span coords
+		const int ixMin = ics * (xMin - hf.bmin[0]);
+		const int izMin = ics * (zMin - hf.bmin[2]);
+		const int ixMax = ics * (xMax - hf.bmin[0]);
+		const int izMax = ics * (zMax - hf.bmin[2]);
+
+		ibminx = ixMin;
+		ibminz = izMin;
+		// the max is inclusive so we need to add 1 to width and height
+		ibwidth = rcClamp(ixMax - ixMin + 1, 0, hf.width);
+		ibheight = rcClamp(izMax - izMin + 1, 0, hf.height);
 		memset(hf.tempSpanColumns, 0, sizeof(rcSpanData)*ibwidth*ibheight);	// we could keep an array of call idx to avoid doing that too often
 #endif // EPIC_ADDITION_USE_NEW_RECAST_RASTERIZER
 	}
@@ -1105,8 +1112,16 @@ struct AddSpanInTempColumns
 		// convert the heightfield coord to the tempSpanColumns coord
 		const int tempColumnx = x - ibminx;
 		const int tempColumnz = z - ibminz;
-		rcAssert(tempColumnx >= 0 && tempColumnx < ibwidth);
-		rcAssert(tempColumnz >= 0 && tempColumnz < ibheight);
+		if (tempColumnx < 0 || tempColumnz < 0 || tempColumnx >= ibwidth || tempColumnz >= ibheight)
+		{
+			// the span is outside of the verts bounding box, we have to skip it because it would mean we access invalid memory
+			if (ctx)
+			{
+				ctx->log(RC_LOG_WARNING, "rcRasterizeTriangles trying to rasterize a triangle as filled convex but some spans are outside of the provided bounding box. Those spans will be skipped.");
+			}
+			return;
+		}
+
 		const int tempColumnIdx = tempColumnx + (tempColumnz) * ibwidth;
 		rcSpanData& tempColumn = hf.tempSpanColumns[tempColumnIdx];
 		if (tempColumn.smin == 0 && tempColumn.smax == 0)
@@ -1147,6 +1162,7 @@ struct AddSpanInTempColumns
 
 	bool IsValid() const { return ibwidth > 0 && ibheight > 0; }
 
+	rcContext* ctx = nullptr;
 	int ibminx = 0;
 	int ibminz = 0;
 	int ibwidth = 0;
