Addresses an incorrect array resizeGrow implementation in Recast helper struct rcScopedDelete. When an array is grown from N to M, M elements were copied from the last array address which read-accesses memory it shouldn't.

diff --git a/Engine/Source/Runtime/Navmesh/Public/Recast/RecastAlloc.h b/Engine/Source/Runtime/Navmesh/Public/Recast/RecastAlloc.h
--- a/Engine/Source/Runtime/Navmesh/Public/Recast/RecastAlloc.h
+++ b/Engine/Source/Runtime/Navmesh/Public/Recast/RecastAlloc.h
@@ -142,7 +142,7 @@ bool rcScopedDelete<T>::resizeGrow(int n)
 	if (n > size)
 	{
 		T* newData = (T*)rcAlloc(sizeof(T) * n, RC_ALLOC_TEMP);
-		if (n && newData) rcMemCpy(newData, ptr, sizeof(T) * n);
+		if (n && newData) rcMemCpy(newData, ptr, sizeof(T) * size);
 		rcFree(ptr);
 		ptr = newData;
 		size = n;
