PR #12629: Recast Prevent Layer IDs from being overwritten in complex tiles

diff --git a/Engine/Source/Runtime/Navmesh/Private/Recast/RecastLayers.cpp b/Engine/Source/Runtime/Navmesh/Private/Recast/RecastLayers.cpp
--- a/Engine/Source/Runtime/Navmesh/Private/Recast/RecastLayers.cpp
+++ b/Engine/Source/Runtime/Navmesh/Private/Recast/RecastLayers.cpp
@@ -699,7 +699,7 @@ static bool SplitAndStoreLayerRegions(rcContext* ctx, rcCompactHeightfield& chf,
 			unsigned short newId = ++layerId;
 			for (int j = i; j < nregs; ++j)
 			{
-				if (regs[j].layerId == oldId)
+				if (regs[j].remap && regs[j].layerId == oldId)
 				{
 					regs[j].layerId = newId;
 					regs[j].remap = false;
