Fix generated navLinks ground validation when using a high JumpEndsHeightTolerance value.

diff --git a/Engine/Source/Runtime/Navmesh/Private/Detour/DetourNavLinkBuilder.cpp b/Engine/Source/Runtime/Navmesh/Private/Detour/DetourNavLinkBuilder.cpp
--- a/Engine/Source/Runtime/Navmesh/Private/Detour/DetourNavLinkBuilder.cpp
+++ b/Engine/Source/Runtime/Navmesh/Private/Detour/DetourNavLinkBuilder.cpp
@@ -629,10 +629,16 @@ void dtNavLinkBuilder::updateTrajectorySamples(EdgeSampler* es) const
 			if (s.floorStart)
 			{
 				s.ymin = (ssmp.height + m_linkBuilderConfig.agentClimb) - p[1];		// -p[1] to stay relative to p[1]
+
+				// Update ymax if ymin is now higher than ymax.
+				s.ymax = dtMax(s.ymin, s.ymax);
 			}
 			else if (s.floorEnd)
 			{
 				s.ymin = (esmp.height + m_linkBuilderConfig.agentClimb) - p[1];		// -p[1] to stay relative to p[1]
+				
+				// Update ymax if ymin is now higher than ymax.
+				s.ymax = dtMax(s.ymin, s.ymax);
 			}
 		}
 	}
@@ -667,7 +673,7 @@ void dtNavLinkBuilder::sampleAction(EdgeSampler* es) const
 	}
 }
 
-void dtNavLinkBuilder::initTrajectorySamples(Trajectory2D* trajectory) const
+void dtNavLinkBuilder::initTrajectorySamples(const float groundRange, Trajectory2D* trajectory) const
 {
 	using namespace UE::Detour::NavLink::Private;
 	
@@ -710,12 +716,18 @@ void dtNavLinkBuilder::initTrajectorySamples(Trajectory2D* trajectory) const
 		s.ymin = dtMin(dtMin(y0,y1), y2) + m_linkBuilderConfig.agentClimb - yRef;	
 		s.ymax = dtMax(dtMax(y0,y1), y2) + m_linkBuilderConfig.agentHeight - yRef; 
 
-		// Mark samples that need to be floored. 
-		// Mainly important for the end where the ground could be far from trajectory end point.
+		// Mark start samples that need to be floored. 
 		if (xRef >= (spine[0]-roundedAgentRadius) && xRef <= spine[0]+roundedAgentRadius)
 			s.floorStart = true;
-		else if (xRef >= (spine[(nspine-1)*2]-roundedAgentRadius) && xRef <= (spine[(nspine-1)*2]+roundedAgentRadius))
+		
+		// More importantly, mark samples that need to be floored at the end since the ground could be far from trajectory end point.
+		// We use the upper bound of the tolerance at the end segment (groundRange) to identify samples that need to be floored.
+		// Min values below the upper bound need to be marked.
+		const float endSplineHeight = pb[1];
+		if (s.ymin + yRef < endSplineHeight + groundRange)
+		{
 			s.floorEnd = true;
+		}		
 	}
 }
 
@@ -1022,7 +1034,7 @@ bool dtNavLinkBuilder::sampleEdge(const dtLinkBuilderConfig& builderConfig, dtNa
 		initJumpOverRig(es, &segs[ibest*6+0], &segs[ibest*6+3], -jumpStartDist, jumpStartDist, jumpHeight, groundRange);
 	}
 
-	initTrajectorySamples(&es->trajectory);
+	initTrajectorySamples(es->groundRange, &es->trajectory);
 	
 	// Init start end segments.
 	dtReal offset[3];
diff --git a/Engine/Source/Runtime/Navmesh/Public/Detour/DetourNavLinkBuilder.h b/Engine/Source/Runtime/Navmesh/Public/Detour/DetourNavLinkBuilder.h
--- a/Engine/Source/Runtime/Navmesh/Public/Detour/DetourNavLinkBuilder.h
+++ b/Engine/Source/Runtime/Navmesh/Public/Detour/DetourNavLinkBuilder.h
@@ -167,7 +167,7 @@ public:
 	NAVMESH_API void debugBuildEdge(const dtLinkBuilderConfig& acfg, dtNavLinkAction action, int edgeIndex, EdgeSampler& sampler);
 	
 private:
-	void initTrajectorySamples(Trajectory2D* trajectory) const;
+	void initTrajectorySamples(const float groundRange, Trajectory2D* trajectory) const;
 	bool isTrajectoryClear(const dtReal* pa, const dtReal* pb, const Trajectory2D* trajectory, const dtReal* trajectoryDir) const;
 	
 	int findPotentialJumpOverEdges(const dtReal* sp, const dtReal* sq,
