Tentative fix for rare crash in storeWallSegment (this will also help track down the source of the issue)

diff --git a/Engine/Source/Runtime/Navmesh/Private/Detour/DetourNavMeshQuery.cpp b/Engine/Source/Runtime/Navmesh/Private/Detour/DetourNavMeshQuery.cpp
--- a/Engine/Source/Runtime/Navmesh/Private/Detour/DetourNavMeshQuery.cpp
+++ b/Engine/Source/Runtime/Navmesh/Private/Detour/DetourNavMeshQuery.cpp
@@ -4500,8 +4500,28 @@ static void storeWallSegment(const dtNavMesh* nav, const dtMeshTile* tile, const
 		return;
 	}
 
-	const dtReal* va = &tile->verts[poly->verts[edge] * 3];
-	const dtReal* vb = &tile->verts[poly->verts[(edge + 1) % poly->vertCount] * 3];
+	if (!ensureMsgf(edge < DT_VERTS_PER_POLYGON, TEXT("Invalid edge index %i"), edge))
+	{
+		return;
+	}
+
+	const int tempEdgeIndexB = (edge + 1) % poly->vertCount;
+	if (!ensureMsgf(tempEdgeIndexB < DT_VERTS_PER_POLYGON, TEXT("Invalid computed edge index %i"), tempEdgeIndexB))
+	{
+		return;
+	}
+
+	const int vertIndexA = poly->verts[edge] * 3;
+	const int vertIndexB = poly->verts[tempEdgeIndexB] * 3;
+	
+	if (!ensureMsgf(vertIndexA < tile->header->vertCount && vertIndexB < tile->header->vertCount,
+		TEXT("Invalid vert index (vertIndexA=%i, vertIndexB=%i, vertCount=%i)"), vertIndexA, vertIndexB, tile->header->vertCount))
+	{
+		return;
+	}
+
+	const dtReal* va = &tile->verts[vertIndexA];
+	const dtReal* vb = &tile->verts[vertIndexB];
 
 	const int32 wall0Offset = (*resultCount * 6) + 0;
 	const int32 wall1Offset = (*resultCount * 6) + 3;
