Fix ensure in storeWallSegment

diff --git a/Engine/Source/Runtime/Navmesh/Private/Detour/DetourNavMeshQuery.cpp b/Engine/Source/Runtime/Navmesh/Private/Detour/DetourNavMeshQuery.cpp
--- a/Engine/Source/Runtime/Navmesh/Private/Detour/DetourNavMeshQuery.cpp
+++ b/Engine/Source/Runtime/Navmesh/Private/Detour/DetourNavMeshQuery.cpp
@@ -4513,8 +4513,9 @@ static void storeWallSegment(const dtNavMesh* nav, const dtMeshTile* tile, const
 
 	const int vertIndexA = poly->verts[edge] * 3;
 	const int vertIndexB = poly->verts[tempEdgeIndexB] * 3;
-	
-	if (!ensureMsgf(vertIndexA < tile->header->vertCount && vertIndexB < tile->header->vertCount,
+
+	const int tileVertArraySize = tile->header->vertCount * 3;
+	if (!ensureMsgf(vertIndexA < tileVertArraySize && vertIndexB < tileVertArraySize,
 		TEXT("Invalid vert index (vertIndexA=%i, vertIndexB=%i, vertCount=%i)"), vertIndexA, vertIndexB, tile->header->vertCount))
 	{
 		return;
diff --git a/Engine/Source/Runtime/Navmesh/Public/Detour/DetourNavMesh.h b/Engine/Source/Runtime/Navmesh/Public/Detour/DetourNavMesh.h
--- a/Engine/Source/Runtime/Navmesh/Public/Detour/DetourNavMesh.h
+++ b/Engine/Source/Runtime/Navmesh/Public/Detour/DetourNavMesh.h
@@ -425,7 +425,7 @@ struct dtMeshTile
 	unsigned int linksFreeList;			///< Index to the next free link.
 	dtMeshHeader* header;				///< The tile header.
 	dtPoly* polys;						///< The tile polygons. [Size: dtMeshHeader::polyCount]
-	dtReal* verts;						///< The tile vertices. [Size: dtMeshHeader::vertCount]
+	dtReal* verts;						///< The tile vertices. [Size: dtMeshHeader::vertCount (3 real per vertex)]
 	dtLink* links;						///< The tile links. [Size: dtMeshHeader::maxLinkCount]
 	dtPolyDetail* detailMeshes;			///< The tile's detail sub-meshes. [Size: dtMeshHeader::detailMeshCount]
 	
