[NavLink Generation] Addition of LinkGenerationSelectedConfig to debug a specific config
Improvements to ground sample debug display

diff --git a/Engine/Source/Runtime/NavigationSystem/Private/NavMesh/RecastNavMesh.cpp b/Engine/Source/Runtime/NavigationSystem/Private/NavMesh/RecastNavMesh.cpp
--- a/Engine/Source/Runtime/NavigationSystem/Private/NavMesh/RecastNavMesh.cpp
+++ b/Engine/Source/Runtime/NavigationSystem/Private/NavMesh/RecastNavMesh.cpp
@@ -461,6 +461,7 @@ FRecastNavMeshTileGenerationDebug::FRecastNavMeshTileGenerationDebug()
 	bUseMaxTileCoordinate = false;
 	LinkGenerationDebugFlags = 0;
 	LinkGenerationSelectedEdge = -1;
+	LinkGenerationSelectedConfig = -1;
 }
 
 #if WITH_RECAST
diff --git a/Engine/Source/Runtime/NavigationSystem/Private/NavMesh/RecastNavMeshGenerator.cpp b/Engine/Source/Runtime/NavigationSystem/Private/NavMesh/RecastNavMeshGenerator.cpp
--- a/Engine/Source/Runtime/NavigationSystem/Private/NavMesh/RecastNavMeshGenerator.cpp
+++ b/Engine/Source/Runtime/NavigationSystem/Private/NavMesh/RecastNavMeshGenerator.cpp
@@ -3934,11 +3934,13 @@ dtStatus FRecastTileGenerator::BuildTileCacheLinks(FNavMeshBuildContext& BuildCo
 	
 	duDebugDraw* dd = nullptr;
 	int32 DebugEdge = -1;
+	int32 DebugConfig = -1;
 
 #if RECAST_INTERNAL_DEBUG_DATA
 	if (IsTileDebugActive() && (TileDebugSettings.LinkGenerationDebugFlags != 0))
 	{
 		DebugEdge = TileDebugSettings.LinkGenerationSelectedEdge;
+		DebugConfig = TileDebugSettings.LinkGenerationSelectedConfig;
 		dd = &BuildContext.InternalDebugData;
 	}
 	const uint16 DebugFlags = TileDebugSettings.LinkGenerationDebugFlags;
@@ -3998,34 +4000,44 @@ dtStatus FRecastTileGenerator::BuildTileCacheLinks(FNavMeshBuildContext& BuildCo
 		LogOnExit();
 		return DT_SUCCESS;	
 	}
+
+	const bool bAllEdges = DebugEdge == -1;
+	const bool bAllConfigs = DebugConfig == -1;
 	
-	if (DebugEdge == -1)
+	if (bAllEdges && bAllConfigs)
 	{
-		rcContext& context = BuildContext;
 		for (const dtNavLinkBuilderJumpConfig& jumpConfig : linkBuilderConfig.jumpConfigs)
 		{
 			if (jumpConfig.enabled)
 			{
 				TRACE_CPUPROFILER_EVENT_SCOPE(RecastBuildLinks_JumpConfig);
-				linkBuilder.buildForAllEdges(context, linkBuilderConfig, jumpConfig);
+				linkBuilder.buildForAllEdges(BuildContext, linkBuilderConfig, jumpConfig);
 			}
 		}
 
 #if RECAST_INTERNAL_DEBUG_DATA
-		duDebugDrawNavLinkBuilder(dd, linkBuilder, DebugFlags, nullptr);
+		duDebugDrawNavLinkBuilder(dd, linkBuilder, DebugFlags, nullptr /*edgeSampler*/);
 	}
 	else
 	{
 		for (const dtNavLinkBuilderJumpConfig& jumpConfig : linkBuilderConfig.jumpConfigs)
 		{
-			if (jumpConfig.enabled)
+			if (jumpConfig.enabled && (bAllConfigs || jumpConfig.ConfigIndex == DebugConfig))
 			{
-				dtNavLinkBuilder::EdgeSampler sampler1;
-				linkBuilder.debugBuildEdge(linkBuilderConfig, DebugEdge, sampler1, jumpConfig);
-				duDebugDrawNavLinkBuilder(dd, linkBuilder, DebugFlags, &sampler1);
+				if (bAllEdges)
+				{
+					linkBuilder.buildForAllEdges(BuildContext, linkBuilderConfig, jumpConfig);
+					duDebugDrawNavLinkBuilder(dd, linkBuilder, DebugFlags, nullptr /*edgeSampler*/);
+				}
+				else
+				{
+					dtNavLinkBuilder::EdgeSampler sampler;
+					linkBuilder.debugBuildEdge(linkBuilderConfig, DebugEdge, sampler, jumpConfig);
+					duDebugDrawNavLinkBuilder(dd, linkBuilder, DebugFlags, &sampler);
+				}
 			}
 		}
-#endif // RECAST_INTERNAL_DEBUG_DATA		
+#endif // RECAST_INTERNAL_DEBUG_DATA
 	}
 
 	AddGeneratedLinks(OutGeneratedLinks, linkBuilderConfig, linkBuilder);
diff --git a/Engine/Source/Runtime/NavigationSystem/Public/NavMesh/RecastNavMesh.h b/Engine/Source/Runtime/NavigationSystem/Public/NavMesh/RecastNavMesh.h
--- a/Engine/Source/Runtime/NavigationSystem/Public/NavMesh/RecastNavMesh.h
+++ b/Engine/Source/Runtime/NavigationSystem/Public/NavMesh/RecastNavMesh.h
@@ -430,10 +430,14 @@ struct FRecastNavMeshTileGenerationDebug
 	UPROPERTY(EditAnywhere, Category = Debug, meta = (Bitmask, BitmaskEnum = "/Script/NavigationSystem.ELinkGenerationDebugFlags"))
 	uint16 LinkGenerationDebugFlags;
 
-	/** Using -1 as no selected edge. */
+	/** Using -1 as 'no edge selected'. */
 	UPROPERTY(EditAnywhere, Category = Debug, meta=(UIMin=-1, ClampMin=-1))
 	int32 LinkGenerationSelectedEdge;
 
+	/** Using -1 as 'no config selected'. */
+	UPROPERTY(EditAnywhere, Category = Debug, meta=(UIMin=-1, ClampMin=-1))
+	int32 LinkGenerationSelectedConfig;
+	
 	/** Indicates if the tile coordinate in parameter is within the tile coordinate selected to debug.
 	 *  Either equal to TileCoordinate or inside the rectangle formed by TileCoordinate and MaxTileCoordinate. */
 	NAVIGATIONSYSTEM_API bool IsWithinTileCoordinates(const int32 TileX, const int32 TileY) const;
diff --git a/Engine/Source/Runtime/Navmesh/Private/DebugUtils/DetourNavLinkDebugDraw.cpp b/Engine/Source/Runtime/Navmesh/Private/DebugUtils/DetourNavLinkDebugDraw.cpp
--- a/Engine/Source/Runtime/Navmesh/Private/DebugUtils/DetourNavLinkDebugDraw.cpp
+++ b/Engine/Source/Runtime/Navmesh/Private/DebugUtils/DetourNavLinkDebugDraw.cpp
@@ -212,7 +212,7 @@ void duDebugDrawTrajectorySamples(duDebugDraw* dd, const dtNavLinkBuilder& linkB
 		else if (s.floorEnd)
 		{
 			dd->begin(DU_DRAW_POINTS, 5.0f);
-			dd->vertex(p[0], p[1]+s.ymin, p[2], duColor::cyan);
+			dd->vertex(p[0], p[1]+s.ymin, p[2], duColor::yellow);
 			dd->end();
 		}
 
@@ -441,6 +441,9 @@ void duDebugDrawNavLinkBuilder(duDebugDraw* dd, const dtNavLinkBuilder& linkBuil
 			dd->vertex(es->rigp[0]+es->az[0]*d, es->rigp[1]+es->az[1]*d, es->rigp[2]+es->az[2]*d, duColor::blue);
 			
 			dd->end();
+
+			const dtNavLinkBuilder::GroundSegment& startSeg = es->start;
+			const dtNavLinkBuilder::GroundSegment& endSeg = es->end;
 			
 			if (drawFlags & DRAW_TRAJECTORY)
 			{
@@ -452,169 +455,84 @@ void duDebugDrawNavLinkBuilder(duDebugDraw* dd, const dtNavLinkBuilder& linkBuil
 				
 				// Start segment.
 				dd->begin(DU_DRAW_LINES, 3.0f);
-				dd->vertex(es->start.p[0],es->start.p[1],es->start.p[2], col);
-				dd->vertex(es->start.q[0],es->start.q[1],es->start.q[2], col);
+				dd->vertex(startSeg.p[0],startSeg.p[1],startSeg.p[2], col);
+				dd->vertex(startSeg.q[0],startSeg.q[1],startSeg.q[2], col);
 				dd->end();
 				
 				dd->begin(DU_DRAW_LINES, 1.0f);
-				dd->vertex(es->start.p[0],es->start.p[1]-r,es->start.p[2], colb);
-				dd->vertex(es->start.p[0],es->start.p[1]+r,es->start.p[2], colb);
-				dd->vertex(es->start.p[0],es->start.p[1]+r,es->start.p[2], colb);
-				dd->vertex(es->start.q[0],es->start.q[1]+r,es->start.q[2], colb);
-				dd->vertex(es->start.q[0],es->start.q[1]+r,es->start.q[2], colb);
-				dd->vertex(es->start.q[0],es->start.q[1]-r,es->start.q[2], colb);
-				dd->vertex(es->start.q[0],es->start.q[1]-r,es->start.q[2], colb);
-				dd->vertex(es->start.p[0],es->start.p[1]-r,es->start.p[2], colb);
+				dd->vertex(startSeg.p[0],startSeg.p[1]-r,startSeg.p[2], colb);
+				dd->vertex(startSeg.p[0],startSeg.p[1]+r,startSeg.p[2], colb);
+				dd->vertex(startSeg.p[0],startSeg.p[1]+r,startSeg.p[2], colb);
+				dd->vertex(startSeg.q[0],startSeg.q[1]+r,startSeg.q[2], colb);
+				dd->vertex(startSeg.q[0],startSeg.q[1]+r,startSeg.q[2], colb);
+				dd->vertex(startSeg.q[0],startSeg.q[1]-r,startSeg.q[2], colb);
+				dd->vertex(startSeg.q[0],startSeg.q[1]-r,startSeg.q[2], colb);
+				dd->vertex(startSeg.p[0],startSeg.p[1]-r,startSeg.p[2], colb);
 				dd->end();
 				
 				// End segment.
 				dd->begin(DU_DRAW_LINES, 3.0f);
-				dd->vertex(es->end.p[0],es->end.p[1],es->end.p[2], col);
-				dd->vertex(es->end.q[0],es->end.q[1],es->end.q[2], col);
+				dd->vertex(endSeg.p[0],endSeg.p[1],endSeg.p[2], col);
+				dd->vertex(endSeg.q[0],endSeg.q[1],endSeg.q[2], col);
 				dd->end();
 				
 				dd->begin(DU_DRAW_LINES, 1.0f);
-				dd->vertex(es->end.p[0],es->end.p[1]-r,es->end.p[2], colb);
-				dd->vertex(es->end.p[0],es->end.p[1]+r,es->end.p[2], colb);
-				dd->vertex(es->end.p[0],es->end.p[1]+r,es->end.p[2], colb);
-				dd->vertex(es->end.q[0],es->end.q[1]+r,es->end.q[2], colb);
-				dd->vertex(es->end.q[0],es->end.q[1]+r,es->end.q[2], colb);
-				dd->vertex(es->end.q[0],es->end.q[1]-r,es->end.q[2], colb);
-				dd->vertex(es->end.q[0],es->end.q[1]-r,es->end.q[2], colb);
-				dd->vertex(es->end.p[0],es->end.p[1]-r,es->end.p[2], colb);
+				dd->vertex(endSeg.p[0],endSeg.p[1]-r,endSeg.p[2], colb);
+				dd->vertex(endSeg.p[0],endSeg.p[1]+r,endSeg.p[2], colb);
+				dd->vertex(endSeg.p[0],endSeg.p[1]+r,endSeg.p[2], colb);
+				dd->vertex(endSeg.q[0],endSeg.q[1]+r,endSeg.q[2], colb);
+				dd->vertex(endSeg.q[0],endSeg.q[1]+r,endSeg.q[2], colb);
+				dd->vertex(endSeg.q[0],endSeg.q[1]-r,endSeg.q[2], colb);
+				dd->vertex(endSeg.q[0],endSeg.q[1]-r,endSeg.q[2], colb);
+				dd->vertex(endSeg.p[0],endSeg.p[1]-r,endSeg.p[2], colb);
 				dd->end();
 		 
 				dd->begin(DU_DRAW_LINES, 4.0f);
-				UE::Detour::Private::drawTrajectory(dd, es->start.p, es->end.p, &es->trajectory, cola);
-				UE::Detour::Private::drawTrajectory(dd, es->start.q, es->end.q, &es->trajectory, cola);
+				UE::Detour::Private::drawTrajectory(dd, startSeg.p, endSeg.p, &es->trajectory, cola);
+				UE::Detour::Private::drawTrajectory(dd, startSeg.q, endSeg.q, &es->trajectory, cola);
 				dd->end();
 
 
 				dd->begin(DU_DRAW_POINTS, 8.0f);
-				dd->vertex(es->start.p[0],es->start.p[1],es->start.p[2], duDarkenCol(col));
-				dd->vertex(es->start.q[0],es->start.q[1],es->start.q[2], duDarkenCol(col));
-				dd->vertex(es->end.p[0],es->end.p[1],es->end.p[2], duDarkenCol(col));
-				dd->vertex(es->end.q[0],es->end.q[1],es->end.q[2], duDarkenCol(col));
+				dd->vertex(startSeg.p[0],startSeg.p[1],startSeg.p[2], duDarkenCol(col));
+				dd->vertex(startSeg.q[0],startSeg.q[1],startSeg.q[2], duDarkenCol(col));
+				dd->vertex(endSeg.p[0],endSeg.p[1],endSeg.p[2], duDarkenCol(col));
+				dd->vertex(endSeg.q[0],endSeg.q[1],endSeg.q[2], duDarkenCol(col));
 				dd->end();
 				
 				unsigned int colm = duColor::white;
 				dd->begin(DU_DRAW_POINTS, 3.0f);
-				dd->vertex(es->start.p[0],es->start.p[1],es->start.p[2], colm);
-				dd->vertex(es->start.q[0],es->start.q[1],es->start.q[2], colm);
-				dd->vertex(es->end.p[0],es->end.p[1],es->end.p[2], colm);
-				dd->vertex(es->end.q[0],es->end.q[1],es->end.q[2], colm);
+				dd->vertex(startSeg.p[0],startSeg.p[1],startSeg.p[2], colm);
+				dd->vertex(startSeg.q[0],startSeg.q[1],startSeg.q[2], colm);
+				dd->vertex(endSeg.p[0],endSeg.p[1],endSeg.p[2], colm);
+				dd->vertex(endSeg.q[0],endSeg.q[1],endSeg.q[2], colm);
 				dd->end();
 			}
 			
 			if (drawFlags & DRAW_LAND_SAMPLES)
 			{
-				if (es->start.ngsamples)
-				{
-					dd->begin(DU_DRAW_POINTS, 8.0f);
-					for (int i = 0; i < es->start.ngsamples; ++i)
-					{
-						const dtNavLinkBuilder::GroundSample* s = &es->start.gsamples[i];
-						const float u = (float)i/(float)(es->start.ngsamples-1);
-						dtReal spt[3], ept[3];
-						dtVlerp(spt, es->start.p, es->start.q, u);
-						dtVlerp(ept, es->start.p, es->start.q, u);
-						
-						unsigned int col = duColor::darkGrey;
-						dtReal off = 1;
-						if ((s->flags & dtNavLinkBuilder::HAS_GROUND) == 0)
-						{
-							off = 0;
-							col = duColor::orangeRed;
-						}
-						
-						spt[1] = s->height + off;
-						dd->vertex(spt, col);
-					}
-					dd->end();
-					
-					dd->begin(DU_DRAW_POINTS, 4.0f);
-					for (int i = 0; i < es->start.ngsamples; ++i)
-					{
-						const dtNavLinkBuilder::GroundSample* s = &es->start.gsamples[i];
-						const float u = (float)i/(float)(es->start.ngsamples-1);
-						dtReal spt[3], ept[3];
-						dtVlerp(spt, es->start.p, es->start.q,u);
-						dtVlerp(ept, es->start.p, es->start.q,u);
-						float off = 0;
-						if (s->flags & dtNavLinkBuilder::HAS_GROUND)
-						{
-							off = 1.f;
-						}
-						spt[1] = s->height + off;
-						dd->vertex(spt, duColor::green);
-					}
-					dd->end();
-					
-				}
-
-				if (es->end.ngsamples)
-				{
-					dd->begin(DU_DRAW_POINTS, 8.0f);
-					for (int i = 0; i < es->end.ngsamples; ++i)
-					{
-						const dtNavLinkBuilder::GroundSample* s = &es->end.gsamples[i];
-						const float u = (float)i/(float)(es->end.ngsamples-1);
-						dtReal spt[3], ept[3];
-						dtVlerp(spt, es->end.p, es->end.q, u);
-						dtVlerp(ept, es->end.p, es->end.q, u);
-						
-						unsigned int col = duColor::darkGrey;
-						float off = 1.f;
-						if ((s->flags & dtNavLinkBuilder::HAS_GROUND) == 0)
-						{
-							off = 0;
-							col = duColor::orangeRed;
-						}
-						
-						spt[1] = s->height + off;
-						dd->vertex(spt, col);
-					}
-					dd->end();
-					
-					dd->begin(DU_DRAW_POINTS, 4.0f);
-					for (int i = 0; i < es->end.ngsamples; ++i)
-					{
-						const dtNavLinkBuilder::GroundSample* s = &es->end.gsamples[i];
-						const float u = (float)i/(float)(es->end.ngsamples-1);
-						dtReal spt[3], ept[3];
-						dtVlerp(spt, es->end.p, es->end.q, u);
-						dtVlerp(ept, es->end.p, es->end.q, u);
-						
-						float off = 0;
-						if (s->flags & dtNavLinkBuilder::HAS_GROUND)
-						{
-							off = 1.f;
-						}
-						spt[1] = s->height + off;
-						dd->vertex(spt, duColor::green);
-					}
-					dd->end();
-				}
+				duDrawGroundSegment(dd, startSeg);
+				duDrawGroundSegment(dd, endSeg);
 			}
 
 			const bool drawCollisionSlices = drawFlags & DRAW_COLLISION_SLICES;
 			const bool drawCollisionSamples = drawFlags & DRAW_COLLISION_SAMPLES;
 			if (drawCollisionSlices || drawCollisionSamples)
 			{
-				if (es->start.ngsamples && es->start.ngsamples == es->end.ngsamples)
+				if (startSeg.ngsamples && startSeg.ngsamples == endSeg.ngsamples)
 				{
-					const int nsamples = es->start.ngsamples;
+					const int nsamples = startSeg.ngsamples;
 					for (int i = 0; i < nsamples; ++i)
 					{
-						const dtNavLinkBuilder::GroundSample* ssmp = &es->start.gsamples[i];
-						const dtNavLinkBuilder::GroundSample* esmp = &es->end.gsamples[i];
+						const dtNavLinkBuilder::GroundSample* ssmp = &startSeg.gsamples[i];
+						const dtNavLinkBuilder::GroundSample* esmp = &endSeg.gsamples[i];
 						if ((ssmp->flags & dtNavLinkBuilder::HAS_GROUND) == 0 || (esmp->flags & dtNavLinkBuilder::HAS_GROUND) == 0)
 							continue;
 						
 						const float u = (float)i/(float)(nsamples-1);
 						dtReal spt[3], ept[3];
-						dtVlerp(spt, es->start.p, es->start.q, u);
-						dtVlerp(ept, es->end.p, es->end.q, u);
+						dtVlerp(spt, startSeg.p, startSeg.q, u);
+						dtVlerp(ept, endSeg.p, endSeg.q, u);
 
 						if (drawCollisionSlices)
 						{
@@ -634,4 +552,38 @@ void duDebugDrawNavLinkBuilder(duDebugDraw* dd, const dtNavLinkBuilder& linkBuil
 		}
 	}
 }
+
+void duDrawGroundSegment(duDebugDraw* dd, const dtNavLinkBuilder::GroundSegment& segment)
+{
+	for (int i = 0; i < segment.ngsamples; ++i)
+	{
+		constexpr float pointSize = 4.f;
+
+		const dtNavLinkBuilder::GroundSample* s = &segment.gsamples[i];
+		
+		const float u = static_cast<float>(i) / static_cast<float>(segment.ngsamples - 1);
+		dtReal spt[3];
+		dtVlerp(spt, segment.p, segment.q, u);
+
+		if (s->flags & dtNavLinkBuilder::HAS_GROUND)
+		{
+			dd->begin(DU_DRAW_POINTS, pointSize);
+			dd->vertex(spt, duColor::lightGreen);
+			dd->vertex(spt[0], s->height, spt[2], duColor::green);
+			dd->end();
+
+			constexpr float lineWidth = 2.f;
+			dd->begin(DU_DRAW_LINES, lineWidth);
+			dd->vertex(spt, duColor::lightGreen);
+			dd->vertex(spt[0], s->height, spt[2], duColor::lightGreen);
+			dd->end();
+		}
+		else
+		{
+			dd->begin(DU_DRAW_POINTS, pointSize);
+			dd->vertex(spt, duColor::orangeRed);
+			dd->end();
+		}
+	}
+}
 //@UE END
diff --git a/Engine/Source/Runtime/Navmesh/Private/Detour/DetourNavLinkBuilder.cpp b/Engine/Source/Runtime/Navmesh/Private/Detour/DetourNavLinkBuilder.cpp
--- a/Engine/Source/Runtime/Navmesh/Private/Detour/DetourNavLinkBuilder.cpp
+++ b/Engine/Source/Runtime/Navmesh/Private/Detour/DetourNavLinkBuilder.cpp
@@ -509,7 +509,7 @@ bool dtNavLinkBuilder::checkHeightfieldCollision(const dtReal x, const dtReal ym
 	return false;
 }
 
-// Returns true if none of the samples ymin, ymax collide with the heghtfield.
+// Returns true if none of the samples ymin, ymax collide with the heightfield.
 bool dtNavLinkBuilder::isTrajectoryClear(const dtReal* pa, const dtReal* pb, const Trajectory2D* trajectory, const dtReal* trajectoryDir) const
 {
 	dtReal start[3];
diff --git a/Engine/Source/Runtime/Navmesh/Public/DebugUtils/DetourNavLinkDebugDraw.h b/Engine/Source/Runtime/Navmesh/Public/DebugUtils/DetourNavLinkDebugDraw.h
--- a/Engine/Source/Runtime/Navmesh/Public/DebugUtils/DetourNavLinkDebugDraw.h
+++ b/Engine/Source/Runtime/Navmesh/Public/DebugUtils/DetourNavLinkDebugDraw.h
@@ -25,4 +25,6 @@ NAVMESH_API void duDebugDrawNavLinkBuilder(duDebugDraw* dd, const dtNavLinkBuild
 
 NAVMESH_API void duDebugDrawTrajectorySamples(duDebugDraw* dd, const dtNavLinkBuilder& linkBuilder, const dtReal* pa, const dtReal* pb,
 											  const dtNavLinkBuilder::Trajectory2D* trajectory, const dtReal* trajectoryDir);
+
+NAVMESH_API void duDrawGroundSegment(duDebugDraw* dd, const dtNavLinkBuilder::GroundSegment& segment);
 //@UE END
diff --git a/Engine/Source/Runtime/Navmesh/Public/Detour/DetourNavLinkBuilder.h b/Engine/Source/Runtime/Navmesh/Public/Detour/DetourNavLinkBuilder.h
--- a/Engine/Source/Runtime/Navmesh/Public/Detour/DetourNavLinkBuilder.h
+++ b/Engine/Source/Runtime/Navmesh/Public/Detour/DetourNavLinkBuilder.h
@@ -148,6 +148,8 @@ private:
 	
 	friend NAVMESH_API void duDebugDrawTrajectorySamples(struct duDebugDraw* dd, const dtNavLinkBuilder& linkBuilder, const dtReal* pa, const dtReal* pb,
 														 const dtNavLinkBuilder::Trajectory2D* trajectory, const dtReal* trajectoryDir);
+
+	friend NAVMESH_API void duDrawGroundSegment(struct duDebugDraw* dd, const dtNavLinkBuilder::GroundSegment& segment);
 	
 public:
 	// Loops through contours to store edge points in world coordinates.
