[NavLink Generation] Don't filter overlapping links if they use different nav areas.

diff --git a/Engine/Source/Runtime/Navmesh/Private/Detour/DetourNavLinkBuilder.cpp b/Engine/Source/Runtime/Navmesh/Private/Detour/DetourNavLinkBuilder.cpp
--- a/Engine/Source/Runtime/Navmesh/Private/Detour/DetourNavLinkBuilder.cpp
+++ b/Engine/Source/Runtime/Navmesh/Private/Detour/DetourNavLinkBuilder.cpp
@@ -372,6 +372,15 @@ void dtNavLinkBuilder::filterOverlappingLinks(const float edgeDistanceThreshold)
 			
 			if (d0 < thresholdSquared && d1 < thresholdSquared && d2 < thresholdSquared && d3 < thresholdSquared)
 			{
+				// If the links are similar but use different areas, they need to be kept.
+				const dtNavLinkBuilderJumpConfig& liConfig = m_linkBuilderConfig.jumpConfigs[li.configIndex];
+				const dtNavLinkBuilderJumpConfig& ljConfig = m_linkBuilderConfig.jumpConfigs[lj.configIndex];
+				if (liConfig.upDirArea != ljConfig.upDirArea ||
+					liConfig.downDirArea != ljConfig.downDirArea)
+				{
+					continue;
+				}
+				
 				// Remove one of the link, keeping the wider one
 				if (dtVdistSqr(spi,sqi) > dtVdistSqr(spj,sqj))
 				{
