[NavLink Generation] Fix missing links on small navmesh patches.

diff --git a/Engine/Source/Runtime/Navmesh/Private/Detour/DetourNavLinkBuilder.cpp b/Engine/Source/Runtime/Navmesh/Private/Detour/DetourNavLinkBuilder.cpp
--- a/Engine/Source/Runtime/Navmesh/Private/Detour/DetourNavLinkBuilder.cpp
+++ b/Engine/Source/Runtime/Navmesh/Private/Detour/DetourNavLinkBuilder.cpp
@@ -272,7 +272,14 @@ void dtNavLinkBuilder::addEdgeLinks(const dtLinkBuilderConfig& builderConfig, co
 		{
 			if (!valid)
 			{
-				const dtReal freeWidth = ((i-start)-1)*distanceBetweenSamples;
+				dtReal freeWidth = ((i-start)-1)*distanceBetweenSamples;
+
+				// Trajectories just before the start or after the end of the edge are expected to be unrestricted.
+				if (start == 0 || i == (nsamples-1))
+				{
+					freeWidth += builderConfig.agentRadius;
+				}
+
 				if (freeWidth >= builderConfig.agentRadius)
 				{
 					const float u0 = (float)start/(float)(nsamples-1);
