Fix FindDistanceToWall returning incorrect result and finding a closest position that is not on a wall.

diff --git a/Engine/Source/Runtime/Navmesh/Private/Detour/DetourNavMeshQuery.cpp b/Engine/Source/Runtime/Navmesh/Private/Detour/DetourNavMeshQuery.cpp
--- a/Engine/Source/Runtime/Navmesh/Private/Detour/DetourNavMeshQuery.cpp
+++ b/Engine/Source/Runtime/Navmesh/Private/Detour/DetourNavMeshQuery.cpp
@@ -328,6 +328,7 @@ dtStatus dtNavMeshQuery::findRandomPoint(const dtQueryFilter* filter, float(*fra
 		if (u*tsum <= area)
 			tile = t;
 	}
+	
 	if (!tile)
 		return DT_FAILURE;
 
@@ -340,9 +341,11 @@ dtStatus dtNavMeshQuery::findRandomPoint(const dtQueryFilter* filter, float(*fra
 	for (int i = 0; i < tile->header->polyCount; ++i)
 	{
 		const dtPoly* p = &tile->polys[i];
+		
 		// Do not return off-mesh connection polygons.
 		if (p->getType() != DT_POLYTYPE_GROUND)
 			continue;
+		
 		// Must pass filter
 		const dtPolyRef ref = base | (dtPolyRef)i;
 		if (!filter->passFilter(ref, tile, p) || !passLinkFilter(tile, i))
@@ -5171,7 +5174,7 @@ dtStatus dtNavMeshQuery::findDistanceToWall(dtPolyRef startRef, const dtReal* ce
 		bestNode->flags |= DT_NODE_CLOSED;
 		
 		// Get poly and tile.
-		// The API input has been cheked already, skip checking internal data.
+		// The API input has been checked already, skip checking internal data.
 		const dtPolyRef bestRef = bestNode->id;
 		const dtMeshTile* bestTile = 0;
 		const dtPoly* bestPoly = 0;
@@ -5193,6 +5196,10 @@ dtStatus dtNavMeshQuery::findDistanceToWall(dtPolyRef startRef, const dtReal* ce
 		// Hit test walls.
 		for (int i = 0, j = (int)bestPoly->vertCount-1; i < (int)bestPoly->vertCount; j = i++)
 		{
+			bool isPortal = false;
+			dtReal portalLeft[3];
+			dtReal portalRight[3];
+			
 			// Skip non-solid edges.
 			if (bestPoly->neis[j] & DT_EXT_LINK)
 			{
@@ -5208,16 +5215,37 @@ dtStatus dtNavMeshQuery::findDistanceToWall(dtPolyRef startRef, const dtReal* ce
 					{
 						if (link.ref != 0)
 						{
+							// Find neighbor tile and poly from the link.ref
 							const dtMeshTile* neiTile = 0;
 							const dtPoly* neiPoly = 0;
 							m_nav->getTileAndPolyByRefUnsafe(link.ref, &neiTile, &neiPoly);
+							
 							if (filter->passFilter(link.ref, neiTile, neiPoly) && passLinkFilterByRef(neiTile, link.ref))
+							{
 								solid = false;
+							}
+							else
+							{
+								// Solid edge, unpack portal limits if needed.
+								if (link.bmin != 0 || link.bmax != 255)
+								{
+									isPortal = true;
+									constexpr dtReal s = 1.0/255.0;
+									const dtReal tmin = link.bmin*s;
+									const dtReal tmax = link.bmax*s;
+									const dtReal* vj = &bestTile->verts[bestPoly->verts[j]*3];
+									const dtReal* vi = &bestTile->verts[bestPoly->verts[i]*3];
+									dtVlerp(portalLeft,  vj, vi, tmin);
+									dtVlerp(portalRight, vj, vi, tmax);
+								}
+							}
 						}
 						break;
 					}
 				}
-				if (!solid) continue;
+
+				if (!solid)
+					continue;
 			}
 			else if (bestPoly->neis[j])
 			{
@@ -5229,10 +5257,21 @@ dtStatus dtNavMeshQuery::findDistanceToWall(dtPolyRef startRef, const dtReal* ce
 			}
 			
 			// Calc distance to the edge.
-			const dtReal* vj = &bestTile->verts[bestPoly->verts[j]*3];
-			const dtReal* vi = &bestTile->verts[bestPoly->verts[i]*3];
+			dtReal* vj;
+			dtReal* vi;
+			if (!isPortal)
+			{
+				vj = &bestTile->verts[bestPoly->verts[j]*3];
+				vi = &bestTile->verts[bestPoly->verts[i]*3];
+			}
+			else
+			{
+				vj = portalLeft;
+				vi = portalRight;
+			}
+
 			dtReal tseg;
-			dtReal distSqr = dtDistancePtSegSqr2D(centerPos, vj, vi, tseg);
+			const dtReal distSqr = dtDistancePtSegSqr2D(centerPos, vj, vi, tseg);
 			
 			// Edge is too far, skip.
 			if (distSqr > radiusSqr)
@@ -5240,12 +5279,13 @@ dtStatus dtNavMeshQuery::findDistanceToWall(dtPolyRef startRef, const dtReal* ce
 			
 			// Hit wall, update radius.
 			radiusSqr = distSqr;
+			
 			// Calculate hit pos.
 			hitPos[0] = vj[0] + (vi[0] - vj[0])*tseg;
 			hitPos[1] = vj[1] + (vi[1] - vj[1])*tseg;
 			hitPos[2] = vj[2] + (vi[2] - vj[2])*tseg;
 		}
-		
+
 		unsigned int i = bestPoly->firstLink;
 		while (i != DT_NULL_LINK)
 		{
@@ -5271,7 +5311,7 @@ dtStatus dtNavMeshQuery::findDistanceToWall(dtPolyRef startRef, const dtReal* ce
 			CA_SUPPRESS(6385);
 			const dtReal* vb = &bestTile->verts[bestPoly->verts[(link.edge+1) % bestPoly->vertCount]*3];
 			dtReal tseg;
-			dtReal distSqr = dtDistancePtSegSqr2D(centerPos, va, vb, tseg);
+			const dtReal distSqr = dtDistancePtSegSqr2D(centerPos, va, vb, tseg);
 			
 			// If the circle is not touching the next polygon, skip it.
 			if (distSqr > radiusSqr)
