[NavLink] Activate the "enable" property now that multiple configurations are supported.
And remove some legacy code.

diff --git a/Engine/Source/Runtime/NavigationSystem/Public/NavMesh/LinkGenerationConfig.h b/Engine/Source/Runtime/NavigationSystem/Public/NavMesh/LinkGenerationConfig.h
--- a/Engine/Source/Runtime/NavigationSystem/Public/NavMesh/LinkGenerationConfig.h
+++ b/Engine/Source/Runtime/NavigationSystem/Public/NavMesh/LinkGenerationConfig.h
@@ -39,7 +39,7 @@ struct FNavLinkGenerationJumpConfig
 	PRAGMA_ENABLE_DEPRECATION_WARNINGS
 
 	/** Should this config be used to generate links. */
-	UPROPERTY(Config)		// @todo Hidden for now. Since there is currently only one config, it has no real usage.
+	UPROPERTY(EditAnywhere, Config, Category = Settings)
 	bool bEnabled = true;
 
 	/** A name for this config */
diff --git a/Engine/Source/Runtime/Navmesh/Private/DebugUtils/DetourNavLinkDebugDraw.cpp b/Engine/Source/Runtime/Navmesh/Private/DebugUtils/DetourNavLinkDebugDraw.cpp
--- a/Engine/Source/Runtime/Navmesh/Private/DebugUtils/DetourNavLinkDebugDraw.cpp
+++ b/Engine/Source/Runtime/Navmesh/Private/DebugUtils/DetourNavLinkDebugDraw.cpp
@@ -310,23 +310,23 @@ void duDebugDrawNavLinkBuilder(duDebugDraw* dd, const dtNavLinkBuilder& linkBuil
 	{
 		const bool drawFilteredLinks = drawFlags & DRAW_FILTERED_LINKS;
 		constexpr int saturation = 50;
-		unsigned int jumpDownCol0 = duLerpCol(duColor::blue, duColor::white, saturation);
-		unsigned int jumpDownCol1 = duColor::blue;
-		unsigned int jumpOverCol0 = duLerpCol(duColor::lightGrey, duColor::white, saturation);
-		unsigned int jumpOverCol1 = duColor::lightGrey;
+		unsigned int jumpConfig0Col0 = duLerpCol(duColor::blue, duColor::white, saturation);
+		unsigned int jumpConfig0Col1 = duColor::blue;
+		unsigned int jumpConfig1Col0 = duLerpCol(duColor::lightGrey, duColor::white, saturation);
+		unsigned int jumpConfig1Col1 = duColor::lightGrey;
 
 		auto selectColors = [&](dtNavLinkBuilder::JumpLinkFlag flag, const uint8 configIndex, unsigned int& col0, unsigned int& col1)
 		{
 			// Predefined "nice" colors for the first two configs. Random colors for the others.
 			if (configIndex == 0)
 			{
-				col0 = jumpDownCol0;
-				col1 = jumpDownCol1;
+				col0 = jumpConfig0Col0;
+				col1 = jumpConfig0Col1;
 			}
 			else if (configIndex == 1)
 			{
-				col0 = jumpOverCol0;
-				col1 = jumpOverCol1;
+				col0 = jumpConfig1Col0;
+				col1 = jumpConfig1Col1;
 			}
 			else
 			{
diff --git a/Engine/Source/Runtime/Navmesh/Private/Detour/DetourNavLinkBuilder.cpp b/Engine/Source/Runtime/Navmesh/Private/Detour/DetourNavLinkBuilder.cpp
--- a/Engine/Source/Runtime/Navmesh/Private/Detour/DetourNavLinkBuilder.cpp
+++ b/Engine/Source/Runtime/Navmesh/Private/Detour/DetourNavLinkBuilder.cpp
@@ -21,69 +21,6 @@ namespace UE::Detour::NavLink::Private
 			a[j+1] = value;
 		}
 	}
-
-	static dtReal getClosestPtPtSeg(const dtReal* pt, const dtReal* sp, const dtReal* sq)
-	{
-		dtReal dir[3], diff[3];
-		dtVsub(dir, sq, sp);
-		dtVsub(diff, pt, sp);
-		dtReal t = dtVdot(dir,diff);
-		if (t <= 0.0)
-			return 0.0;
-	
-		const dtReal d = dtVdot(dir,dir);
-		if (t >= d)
-			return 1.0;
-
-		return t/d;
-	}
-
-	static bool isectSegAABB(const dtReal* sp, const dtReal* sq,
-					  const float* amin, const float* amax,
-					  float& tmin, float& tmax)
-	{
-		static constexpr float EPS = 1e-6f;
-	
-		dtReal d[3];
-		dtVsub(d, sq, sp);
-		tmin = 0;  // set to -FLT_MAX to get first hit on the line
-		tmax = FLT_MAX;		// set to max distance ray can travel (for segment)
-	
-		// For all three slabs
-		for (int i = 0; i < 3; i++)
-		{
-			if (dtAbs(d[i]) < EPS)
-			{
-				// Ray is parallel to slab. No hit if origin not within slab
-				if (sp[i] < amin[i] || sp[i] > amax[i])
-					return false;
-			}
-			else
-			{
-				// Compute intersection t value of ray with near and far plane of slab
-				const float ood = 1.0f / d[i];
-				float t1 = (amin[i] - sp[i]) * ood;
-				float t2 = (amax[i] - sp[i]) * ood;
-
-				// Make t1 be intersection with near plane, t2 with far plane
-				if (t1 > t2)
-					dtSwap(t1, t2);
-
-				// Compute the intersection of slab intersections intervals
-				if (t1 > tmin)
-					tmin = t1;
-
-				if (t2 < tmax) 
-					tmax = t2;
-
-				// Exit with no collision as soon as slab intersection becomes empty
-				if (tmin > tmax)
-					return false;
-			}
-		}
-	
-		return true;
-	}
 	
 	static float getHeight(const float x, const float* pts, const int npts)
 	{
@@ -730,178 +667,6 @@ void dtNavLinkBuilder::initTrajectorySamples(const float groundRange, Trajectory
 	}
 }
 
-int dtNavLinkBuilder::findPotentialJumpOverEdges(const dtReal* sp, const dtReal* sq,
-												  const float depthRange, const float heightRange,
-												  dtReal* outSegs, const int maxOutSegs) const
-{
-	TRACE_CPUPROFILER_EVENT_SCOPE(dtNavLinkBuilder::findPotentialJumpOverEdges);
-	
-	using namespace UE::Detour::NavLink::Private;
-	
-	// Find potential edges to join to.
-	const float widthRange = sqrtf(dtVdistSqr(sp,sq));
-	const float amin[3] = { 0, -heightRange*0.5f, 0 };
-	const float amax[3] = { widthRange, heightRange*0.5f, depthRange};
-	
-	const dtReal thr = cosf((180.0 - 45.0)/180.0*RC_PI);
-	
-	dtReal ax[3], ay[3], az[3];
-	dtVsub(ax, sq, sp);
-	dtVnormalize(ax);
-	dtVset(az, ax[2], 0, -ax[0]);
-	dtVnormalize(az);
-	dtVset(ay, 0, 1, 0);
-	
-	static constexpr int MAX_SEGS = 64;
-	PotentialSeg segs[MAX_SEGS];
-	int nsegs = 0;
-	
-	for (int i = 0; i < m_edges.Num(); ++i)
-	{
-		dtReal p[3], lsp[3], lsq[3];
-		dtVsub(p, m_edges[i].sp, sp);
-		lsp[0] = ax[0]*p[0] + ay[0]*p[1] + az[0]*p[2];
-		lsp[1] = ax[1]*p[0] + ay[1]*p[1] + az[1]*p[2];
-		lsp[2] = ax[2]*p[0] + ay[2]*p[1] + az[2]*p[2];
-		
-		dtVsub(p, m_edges[i].sq, sp);
-		lsq[0] = ax[0]*p[0] + ay[0]*p[1] + az[0]*p[2];
-		lsq[1] = ax[1]*p[0] + ay[1]*p[1] + az[1]*p[2];
-		lsq[2] = ax[2]*p[0] + ay[2]*p[1] + az[2]*p[2];
-		
-		float tmin, tmax;
-		if (isectSegAABB(lsp, lsq, amin, amax, tmin, tmax))
-		{
-			if (tmin > 1.0f)
-				continue;
-			if (tmax < 0.0f)
-				continue;
-			
-			dtReal edir[3];
-			dtVsub(edir, m_edges[i].sq, m_edges[i].sp);
-			edir[1] = 0;
-			dtVnormalize(edir);
-			
-			if (dtVdot(ax, edir) > thr)
-				continue;
-			
-			if (nsegs < MAX_SEGS)
-			{
-				segs[nsegs].umin = dtClamp(tmin, 0.0f, 1.0f);
-				segs[nsegs].umax = dtClamp(tmax, 0.0f, 1.0f);
-				segs[nsegs].dmin = dtMin(lsp[2], lsq[2]);
-				segs[nsegs].dmax = dtMax(lsp[2], lsq[2]);
-				segs[nsegs].idx = i;
-				segs[nsegs].mark = 0;
-				nsegs++;
-			}
-		}
-	}
-	
-	const float eps = m_chf->cs;
-	unsigned char mark = 1;
-	for (int i = 0; i < nsegs; ++i)
-	{
-		if (segs[i].mark != 0)
-			continue;
-
-		segs[i].mark = mark;
-		
-		for (int j = i+1; j < nsegs; ++j)
-		{
-			if (overlapRange(segs[i].dmin-eps, segs[i].dmax+eps,
-						 	 segs[j].dmin-eps, segs[j].dmax+eps))
-			{
-				segs[j].mark = mark;
-			}
-		}
-		
-		mark++;
-	}
-
-
-	int nout = 0;
-	
-	for (int i = 1; i < mark; ++i)
-	{
-		// Find destination mid point.
-		constexpr float toUU = 100.f;
-		float umin = 10.0f * toUU;
-		float umax = -10.0f * toUU;
-		dtReal ptmin[3], ptmax[3];
-		
-		for (int j = 0; j < nsegs; ++j)
-		{
-			PotentialSeg* seg = &segs[j];
-			if (seg->mark != (unsigned char)i)
-				continue;
-			
-			dtReal pa[3], pb[3];
-			dtVlerp(pa, m_edges[seg->idx].sp, m_edges[seg->idx].sq, seg->umin);
-			dtVlerp(pb, m_edges[seg->idx].sp, m_edges[seg->idx].sq, seg->umax);
-			
-			const float ua = getClosestPtPtSeg(pa, sp, sq);
-			const float ub = getClosestPtPtSeg(pb, sp, sq);
-			
-			if (ua < umin)
-			{
-				dtVcopy(ptmin, pa);
-				umin = ua;
-			}
-			if (ua > umax)
-			{
-				dtVcopy(ptmax, pa);
-				umax = ua;
-			}
-			
-			if (ub < umin)
-			{
-				dtVcopy(ptmin, pb);
-				umin = ub;
-			}
-			if (ub > umax)
-			{
-				dtVcopy(ptmax, pb);
-				umax = ub;
-			}
-		}
-		
-		if (umin > umax)
-			continue;
-		
-		dtReal end[3];
-		dtVlerp(end, ptmin, ptmax, 0.5f);
-		
-		dtReal start[3];
-		dtVlerp(start, sp, sq, (umin+umax)*0.5f);
-		
-		dtReal orig[3];
-		dtVlerp(orig, start, end, 0.5f);
-		
-		dtReal dir[3], norm[3];
-		dtVsub(dir, end, start);
-		dir[1] = 0;
-		dtVnormalize(dir);
-		dtVset(norm, dir[2], 0, -dir[0]);
-		
-		dtReal ssp[3], ssq[3];
-		
-		const dtReal width = widthRange * (umax-umin);
-		dtVmad(ssp, orig, norm, width*0.5f);
-		dtVmad(ssq, orig, norm, -width*0.5f);
-		
-		if (nout < maxOutSegs)
-		{
-			dtVcopy(&outSegs[nout*6+0], ssp);
-			dtVcopy(&outSegs[nout*6+3], ssq);
-			nout++;
-		}
-	}
-		
-	return nout;
-}
-	
-
 void dtNavLinkBuilder::initJumpRig(EdgeSampler* es, const dtReal* sp, const dtReal* sq, const dtNavLinkBuilderJumpConfig& config) const
 {
 	es->configIndex = config.ConfigIndex;
diff --git a/Engine/Source/Runtime/Navmesh/Public/Detour/DetourNavLinkBuilder.h b/Engine/Source/Runtime/Navmesh/Public/Detour/DetourNavLinkBuilder.h
--- a/Engine/Source/Runtime/Navmesh/Public/Detour/DetourNavLinkBuilder.h
+++ b/Engine/Source/Runtime/Navmesh/Public/Detour/DetourNavLinkBuilder.h
@@ -166,10 +166,6 @@ private:
 	void initTrajectorySamples(const float groundRange, Trajectory2D* trajectory) const;
 	bool isTrajectoryClear(const dtReal* pa, const dtReal* pb, const Trajectory2D* trajectory, const dtReal* trajectoryDir) const;
 	
-	int findPotentialJumpOverEdges(const dtReal* sp, const dtReal* sq,
-								   const float depthRange, const float heightRange,
-								   dtReal* outSegs, const int maxOutSegs) const;
-	
 	void initJumpRig(EdgeSampler* es, const dtReal* sp, const dtReal* sq,
 						 const dtNavLinkBuilderJumpConfig& config) const;
 
