Fix for recast not putting data where the caller is expecting.

#Synced-CL 48579635

diff --git a/Engine/Source/Runtime/Navmesh/Private/Recast/RecastLayers.cpp b/Engine/Source/Runtime/Navmesh/Private/Recast/RecastLayers.cpp
--- a/Engine/Source/Runtime/Navmesh/Private/Recast/RecastLayers.cpp
+++ b/Engine/Source/Runtime/Navmesh/Private/Recast/RecastLayers.cpp
@@ -978,7 +978,7 @@ bool rcBuildHeightfieldLayersChunky(rcContext* ctx, rcCompactHeightfield& chf,
 }
 
 /// helper function from RecastRegion.cpp, requires distance data in compact height field
-bool rcGatherRegionsNoFilter(rcContext* ctx, rcCompactHeightfield& chf, const rcBorderSize borderSize, unsigned short* spanBuf4);	//@UE
+bool rcGatherRegionsNoFilter(rcContext* ctx, rcCompactHeightfield& chf, const rcBorderSize borderSize, unsigned short* spanBuf4, unsigned short*& resultBufPtr);	//@UE
 
 struct rcLayerRegion
 {
@@ -1126,7 +1126,8 @@ bool rcBuildHeightfieldLayers(rcContext* ctx, rcCompactHeightfield& chf,
 	ctx->startTimer(RC_TIMER_BUILD_REGIONS_WATERSHED);
 
 	unsigned short* srcReg = spanBuf4;
-	if (!rcGatherRegionsNoFilter(ctx, chf, borderSize, spanBuf4))
+	// srcReg might be updated and point at a different location.
+	if (!rcGatherRegionsNoFilter(ctx, chf, borderSize, spanBuf4, srcReg))	//@UE
 		return false;
 
 	ctx->stopTimer(RC_TIMER_BUILD_REGIONS_WATERSHED);
diff --git a/Engine/Source/Runtime/Navmesh/Private/Recast/RecastRegion.cpp b/Engine/Source/Runtime/Navmesh/Private/Recast/RecastRegion.cpp
--- a/Engine/Source/Runtime/Navmesh/Private/Recast/RecastRegion.cpp
+++ b/Engine/Source/Runtime/Navmesh/Private/Recast/RecastRegion.cpp
@@ -1233,7 +1233,7 @@ bool rcBuildRegionsChunky(rcContext* ctx, rcCompactHeightfield& chf,
 	rcScopedDelete<unsigned short> srcReg = (unsigned short*)rcAlloc(sizeof(unsigned short)*chf.spanCount, RC_ALLOC_TEMP);
 	if (!srcReg)
 	{
-		ctx->log(RC_LOG_ERROR, "rcBuildRegionsMonotone: Out of memory 'src' (%d).", chf.spanCount);
+		ctx->log(RC_LOG_ERROR, "rcBuildRegionsChunky: Out of memory 'src' (%d).", chf.spanCount);
 		return false;
 	}
 	memset(srcReg,0,sizeof(unsigned short)*chf.spanCount);
@@ -1242,7 +1242,7 @@ bool rcBuildRegionsChunky(rcContext* ctx, rcCompactHeightfield& chf,
 	rcScopedDelete<rcSweepSpan> sweeps = (rcSweepSpan*)rcAlloc(sizeof(rcSweepSpan)*nsweeps, RC_ALLOC_TEMP);
 	if (!sweeps)
 	{
-		ctx->log(RC_LOG_ERROR, "rcBuildRegionsMonotone: Out of memory 'sweeps' (%d).", nsweeps);
+		ctx->log(RC_LOG_ERROR, "rcBuildRegionsChunky: Out of memory 'sweeps' (%d).", nsweeps);
 		return false;
 	}
 
@@ -1369,8 +1369,10 @@ bool rcBuildRegionsChunky(rcContext* ctx, rcCompactHeightfield& chf,
 ///       it allows reusing existing code for building layer set (flood fill based partitioning)
 ///		  
 ///		  spanBuf4 is temporary buffer, allocated and freed by caller (size = chf.spanCount * 4)
+/// 
+///		  OutRegPtr points to where the caller should look for the results, which will be at either spanBuf4 or spanBuf4+chf.spanCount*2 //@UE
 ///
-bool rcGatherRegionsNoFilter(rcContext* ctx, rcCompactHeightfield& chf, const rcBorderSize borderSize, unsigned short* spanBuf4)		//@UE
+bool rcGatherRegionsNoFilter(rcContext* ctx, rcCompactHeightfield& chf, const rcBorderSize borderSize, unsigned short* spanBuf4, unsigned short*& resultBufPtr)		//@UE
 {
 	const int w = chf.width;
 	const int h = chf.height;
@@ -1440,6 +1442,9 @@ bool rcGatherRegionsNoFilter(rcContext* ctx, rcCompactHeightfield& chf, const rc
 		rcSwap(srcDist, dstDist);
 	}
 
+	// Because srcReg may not still be pointing at spanBuf4.
+	resultBufPtr = srcReg;			//@UE
+
 	chf.maxRegions = regionId;
 	return true;
 }
@@ -1453,7 +1458,7 @@ bool rcGatherRegionsNoFilter(rcContext* ctx, rcCompactHeightfield& chf, const rc
 /// re-assigned to the zero (null) region.
 /// 
 /// Watershed partitioning can result in smaller than necessary regions, especially in diagonal corridors. 
-/// @p mergeRegionArea helps reduce unecessarily small regions.
+/// @p mergeRegionArea helps reduce unnecessarily small regions.
 /// 
 /// See the #rcConfig documentation for more information on the configuration parameters.
 /// 
@@ -1480,7 +1485,8 @@ bool rcBuildRegions(rcContext* ctx, rcCompactHeightfield& chf,
 	ctx->startTimer(RC_TIMER_BUILD_REGIONS_WATERSHED);
 
 	unsigned short* srcReg = spanBuf4;
-	if (!rcGatherRegionsNoFilter(ctx, chf, borderSize, spanBuf4))
+	// srcReg might be updated and point at a different location.
+	if (!rcGatherRegionsNoFilter(ctx, chf, borderSize, spanBuf4, srcReg))	//@UE
 		return false;
 
 	ctx->stopTimer(RC_TIMER_BUILD_REGIONS_WATERSHED);	
