Handle out of memory condition in dtNavMeshQuery::queryPolygons

#Synced-CL 49035158

diff --git a/Engine/Source/Runtime/Navmesh/Private/Detour/DetourNavMesh.cpp b/Engine/Source/Runtime/Navmesh/Private/Detour/DetourNavMesh.cpp
--- a/Engine/Source/Runtime/Navmesh/Private/Detour/DetourNavMesh.cpp
+++ b/Engine/Source/Runtime/Navmesh/Private/Detour/DetourNavMesh.cpp
@@ -2087,8 +2087,16 @@ void dtNavMesh::processSegmentLinksForTile(dtTileRef tileRef, unsigned int maxSk
 
 	// Connect with layers in current tile.
 	UE_LOG(LogSegmentLink, Log, TEXT("   %hs Connect with layers in current tile."), __FUNCTION__);
+
 	int nneis = getTileCountAt(tileX, tileY);
 	dtMeshTile** neis = TileArray.PrepareArray(nneis);
+	if (neis == nullptr)
+	{
+		UE_LOG(LogSegmentLink, Error, TEXT("   %hs Out of memory."), __FUNCTION__);
+		dtFree(segList, DT_ALLOC_TEMP);
+		return;
+	}
+	
 	getTilesAt(tileX, tileY, neis, nneis);
 	for (int j = 0; j < nneis; ++j)
 	{
@@ -2118,6 +2126,13 @@ void dtNavMesh::processSegmentLinksForTile(dtTileRef tileRef, unsigned int maxSk
 	{
 		nneis = getNeighbourTilesCountAt(tileX, tileY, i);
 		neis = TileArray.PrepareArray(nneis);
+		if (neis == nullptr)
+		{
+			UE_LOG(LogSegmentLink, Error, TEXT("   %hs Out of memory."), __FUNCTION__);
+			dtFree(segList, DT_ALLOC_TEMP);
+			return;
+		}
+		
 		getNeighbourTilesAt(tileX, tileY, i, neis, nneis);
 
 		for (int j = 0; j < nneis; ++j)
diff --git a/Engine/Source/Runtime/Navmesh/Private/Detour/DetourNavMeshQuery.cpp b/Engine/Source/Runtime/Navmesh/Private/Detour/DetourNavMeshQuery.cpp
--- a/Engine/Source/Runtime/Navmesh/Private/Detour/DetourNavMeshQuery.cpp
+++ b/Engine/Source/Runtime/Navmesh/Private/Detour/DetourNavMeshQuery.cpp
@@ -1546,8 +1546,12 @@ dtStatus dtNavMeshQuery::queryPolygons(const dtReal* center, const dtReal* exten
 		{
 			int nneis = m_nav->getTileCountAt(x,y);
 			const dtMeshTile** neis = (const dtMeshTile**)TileArray.PrepareArray(nneis);
+			if (neis == nullptr)						//@UE
+			{											//@UE
+				return DT_FAILURE | DT_OUT_OF_MEMORY;	//@UE
+			}											//@UE
 
-			m_nav->getTilesAt(x,y,neis,nneis);
+			m_nav->getTilesAt(x, y, neis, nneis);
 			for (int j = 0; j < nneis; ++j)
 			{
 				n += queryPolygonsInTile(neis[j], bmin, bmax, filter, polys+n, maxPolys-n);
diff --git a/Engine/Source/Runtime/Navmesh/Public/Detour/DetourNavMesh.h b/Engine/Source/Runtime/Navmesh/Public/Detour/DetourNavMesh.h
--- a/Engine/Source/Runtime/Navmesh/Public/Detour/DetourNavMesh.h
+++ b/Engine/Source/Runtime/Navmesh/Public/Detour/DetourNavMesh.h
@@ -1006,7 +1006,7 @@ struct ReadTilesHelper
 	int NumAllocated;
 	dtMeshTile** AllocatedTiles;
 
-	ReadTilesHelper() : NumAllocated(0), AllocatedTiles(0) {}
+	ReadTilesHelper() : NumAllocated(0), AllocatedTiles(nullptr) {}
 	~ReadTilesHelper()
 	{
 		dtFree(AllocatedTiles, DT_ALLOC_TEMP);
