Fix to make sure that the area costs are set before NavLinks are snapped when serializing in a level.

#Synced-CL 50912775

diff --git a/Engine/Source/Runtime/NavigationSystem/Private/NavMesh/PImplRecastNavMesh.cpp b/Engine/Source/Runtime/NavigationSystem/Private/NavMesh/PImplRecastNavMesh.cpp
--- a/Engine/Source/Runtime/NavigationSystem/Private/NavMesh/PImplRecastNavMesh.cpp
+++ b/Engine/Source/Runtime/NavigationSystem/Private/NavMesh/PImplRecastNavMesh.cpp
@@ -690,6 +690,12 @@ void FPImplRecastNavMesh::Serialize( FArchive& Ar, int32 NavMeshVersion )
 			
 			NavMeshOwner->bHasNoTileData = (NumTiles == 0);	
 
+			// Defer off-mesh point link snapping until after all tiles are loaded
+			// and area cost order is applied. Without this, links using
+			// bSnapToCheapestArea snap by distance instead of cost because
+			// m_areaCostOrder hasn't been initialized yet.
+			DetourNavMesh->setDeferOffMeshLinking(true);
+
 #if WITH_NAVMESH_SEGMENT_LINKS			
 			// Store the tileRefs when iterating over the tile list.
 			// These will be processed again to add the segment link connections.
@@ -3521,6 +3527,10 @@ void FPImplRecastNavMesh::OnAreaCostChanged()
 		}
 
 		DetourNavMesh->applyAreaCostOrder(AreaCostOrder);
+
+		// If off-mesh linking was deferred during batch tile loading,
+		// connect them now that the cost order is correct.
+		DetourNavMesh->connectDeferredOffMeshLinks();
 	}
 }
 
diff --git a/Engine/Source/Runtime/NavigationSystem/Private/NavMesh/RecastNavMesh.cpp b/Engine/Source/Runtime/NavigationSystem/Private/NavMesh/RecastNavMesh.cpp
--- a/Engine/Source/Runtime/NavigationSystem/Private/NavMesh/RecastNavMesh.cpp
+++ b/Engine/Source/Runtime/NavigationSystem/Private/NavMesh/RecastNavMesh.cpp
@@ -1216,6 +1216,10 @@ void ARecastNavMesh::OnRegistered()
 	// This check can fail when the NavMeshVersion indicates the map needs the nav mesh rebuilt
 	ensure(RecastNavMeshImpl->GetRecastMesh() == nullptr || RecastNavMeshImpl->GetRecastMesh()->getBVQuantFactor((uint8)ENavigationDataResolution::Default) != 0);
 
+	// Off-mesh link deferred connection should have been resolved by OnAreaCostChanged
+	// during ProcessNavAreas, which runs before OnRegistered.
+	ensure(RecastNavMeshImpl->GetRecastMesh() == nullptr || !RecastNavMeshImpl->GetRecastMesh()->isDeferringOffMeshLinking());
+
 	Super::OnRegistered();
 
 	// Re-register any generated link proxies that couldn't be registered
diff --git a/Engine/Source/Runtime/Navmesh/Private/Detour/DetourNavMesh.cpp b/Engine/Source/Runtime/Navmesh/Private/Detour/DetourNavMesh.cpp
--- a/Engine/Source/Runtime/Navmesh/Private/Detour/DetourNavMesh.cpp
+++ b/Engine/Source/Runtime/Navmesh/Private/Detour/DetourNavMesh.cpp
@@ -921,6 +921,7 @@ dtNavMesh::dtNavMesh() :
 {
 	memset(&m_params, 0, sizeof(dtNavMeshParams));
 	memset(m_areaCostOrder, 0, sizeof(m_areaCostOrder));
+	m_bDeferOffMeshLinking = false;
 	m_orig[0] = 0;
 	m_orig[1] = 0;
 	m_orig[2] = 0;
@@ -2012,7 +2013,10 @@ dtStatus dtNavMesh::addTile(unsigned char* data, int dataSize, int flags,
 	tile->flags = flags;
 
 	connectIntLinks(tile);
-	baseOffMeshLinks(tile);
+	if (!m_bDeferOffMeshLinking)
+	{
+		baseOffMeshLinks(tile);
+	}
 
 	// Create connections with neighbour tiles.
 	ReadTilesHelper TileArray;
@@ -2034,9 +2038,15 @@ dtStatus dtNavMesh::addTile(unsigned char* data, int dataSize, int flags,
 		{
 			connectExtLinks(tile, neis[j], -1, bHasClusters);
 			connectExtLinks(neis[j], tile, -1, bHasClusters);
-			connectExtOffMeshLinks(tile, neis[j], -1, bHasClusters);
+			if (!m_bDeferOffMeshLinking)
+			{
+				connectExtOffMeshLinks(tile, neis[j], -1, bHasClusters);
+			}
+		}
+		if (!m_bDeferOffMeshLinking)
+		{
+			connectExtOffMeshLinks(neis[j], tile, -1, bHasClusters);
 		}
-		connectExtOffMeshLinks(neis[j], tile, -1, bHasClusters);
 	}
 	
 	// Connect with neighbour tiles.
@@ -2061,8 +2071,11 @@ dtStatus dtNavMesh::addTile(unsigned char* data, int dataSize, int flags,
 				connectExtLinks(neis[j], tile, dtOppositeTile(i), bHasClusters);
 			}
 
-			connectExtOffMeshLinks(tile, neis[j], i, bHasClusters);
-			connectExtOffMeshLinks(neis[j], tile, dtOppositeTile(i), bHasClusters);
+			if (!m_bDeferOffMeshLinking)
+			{
+				connectExtOffMeshLinks(tile, neis[j], i, bHasClusters);
+				connectExtOffMeshLinks(neis[j], tile, dtOppositeTile(i), bHasClusters);
+			}
 		}
 	}
 	
@@ -2073,6 +2086,63 @@ dtStatus dtNavMesh::addTile(unsigned char* data, int dataSize, int flags,
 	
 	return DT_SUCCESS;
 }
+
+void dtNavMesh::connectDeferredOffMeshLinks()
+{
+	if (!m_bDeferOffMeshLinking)
+		return;
+
+	// Called from OnAreaCostChanged after applyAreaCostOrder has set m_areaCostOrder.
+	// Must not be called before area costs are applied, or bSnapToCheapestArea links
+	// will snap by distance instead of cost.
+
+	m_bDeferOffMeshLinking = false;
+
+	ReadTilesHelper TileArray;
+
+	for (int i = 0; i < m_maxTiles; ++i)
+	{
+		dtMeshTile* tile = &m_tiles[i];
+		if (!tile->header) continue;
+
+		const int tx = tile->header->x;
+		const int ty = tile->header->y;
+
+#if WITH_NAVMESH_CLUSTER_LINKS
+		const bool bHasClusters = tile->header->clusterCount > 0;
+#else
+		const bool bHasClusters = false;
+#endif // WITH_NAVMESH_CLUSTER_LINKS
+		
+		// Snap start points (vertex 0) of this tile's off-mesh connections
+		baseOffMeshLinks(tile);
+
+		// Connect layers in same column — landing on this tile
+		int nneis = getTileCountAt(tx, ty);
+		dtMeshTile** neis = TileArray.PrepareArray(nneis);
+		if (!neis) continue;
+		getTilesAt(tx, ty, neis, nneis);
+		for (int j = 0; j < nneis; ++j)
+		{
+			// neis[j]'s connections landing on tile (includes self case)
+			connectExtOffMeshLinks(tile, neis[j], -1, bHasClusters);
+		}
+
+		// Connect neighbour tiles across edges — landing on this tile
+		for (int side = 0; side < 8; ++side)
+		{
+			nneis = getNeighbourTilesCountAt(tx, ty, side);
+			neis = TileArray.PrepareArray(nneis);
+			if (!neis) continue;
+			getNeighbourTilesAt(tx, ty, side, neis, nneis);
+			for (int j = 0; j < nneis; ++j)
+			{
+				// neis[j]'s connections landing on tile
+				connectExtOffMeshLinks(tile, neis[j], side, bHasClusters);
+			}
+		}
+	}
+}
 //@UE END
 
 #if WITH_NAVMESH_SEGMENT_LINKS
diff --git a/Engine/Source/Runtime/Navmesh/Public/Detour/DetourNavMesh.h b/Engine/Source/Runtime/Navmesh/Public/Detour/DetourNavMesh.h
--- a/Engine/Source/Runtime/Navmesh/Public/Detour/DetourNavMesh.h
+++ b/Engine/Source/Runtime/Navmesh/Public/Detour/DetourNavMesh.h
@@ -865,6 +865,15 @@ public:
 
 	NAVMESH_API void applyAreaCostOrder(unsigned char* costOrder);
 	
+	/// Defers off-mesh point link connection during addTile.
+	/// Call connectDeferredOffMeshLinks() after all tiles are loaded.
+	void setDeferOffMeshLinking(bool bDefer) { m_bDeferOffMeshLinking = bDefer; }
+	bool isDeferringOffMeshLinking() const { return m_bDeferOffMeshLinking; }
+
+	/// Connects off-mesh point links for all tiles. Used after batch loading
+	/// when off-mesh linking was deferred via setDeferOffMeshLinking(true).
+	NAVMESH_API void connectDeferredOffMeshLinks();
+
 	/// Returns neighbour tile count based on side of given tile.
 	NAVMESH_API int getNeighbourTilesCountAt(const int x, const int y, const int side) const;
 
@@ -984,6 +993,7 @@ public:
 	int m_tileLutMask;					///< Tile hash lookup mask.
 
 	unsigned char m_areaCostOrder[DT_MAX_AREAS];
+	bool m_bDeferOffMeshLinking;
 
 	dtMeshTile** m_posLookup;			///< Tile hash lookup.
 	dtMeshTile* m_nextFree;				///< Freelist of tiles.
