[Navigation] Optimize navmesh vertex storage with tile-local coordinates
Convert navmesh vertex data from world-space doubles to tile-local float offsets, reducing memory usage by ~50% (from 24 bytes/vertex to 12 bytes/vertex).
- Introduce dtTileVert struct to store vertices as offsets from tile bmin
- Add methods to converting between local to world space
- Update all vertex access throughout NavigationSystem and Detour to use new encoding
- Bump navmesh version to 28 (NAVMESHVER_TILEVERT_ENCODING), DT_NAVMESH_VERSION to 8
- Add backward compatibility: old format data is automatically converted on load
- Add ai.debug.nav.DumpNavMeshMemory console command outputing a detailed memory report
- Few additional fixes
#Synced-CL 51761695

diff --git a/Engine/Source/Runtime/NavigationSystem/Private/NavMesh/PImplRecastNavMesh.cpp b/Engine/Source/Runtime/NavigationSystem/Private/NavMesh/PImplRecastNavMesh.cpp
--- a/Engine/Source/Runtime/NavigationSystem/Private/NavMesh/PImplRecastNavMesh.cpp
+++ b/Engine/Source/Runtime/NavigationSystem/Private/NavMesh/PImplRecastNavMesh.cpp
@@ -826,6 +826,7 @@ void FPImplRecastNavMesh::SerializeRecastMeshTile(FArchive& Ar, int32 NavMeshVer
 	// calc sizes for our data so we know how much to allocate and where to read/write stuff
 	// note this may not match the on-disk size or the in-memory size on the machine that generated that data
 
+	// In-memory tile layout always uses latest vertex format (dtTileVert).
 	FDetourTileLayout TileLayout(SizeInfo);
 
 	if (Ar.IsLoading())
@@ -839,7 +840,10 @@ void FPImplRecastNavMesh::SerializeRecastMeshTile(FArchive& Ar, int32 NavMeshVer
 		{
 			UE_LOG(LogNavigation, Error, TEXT("Failed to alloc navmesh tile"));
 		}
-		FMemory::Memset(TileData, 0, TileDataSize);
+		else
+		{
+			FMemory::Memset(TileData, 0, TileDataSize);
+		}
 	}
 	else if (Ar.IsSaving())
 	{
@@ -853,11 +857,11 @@ void FPImplRecastNavMesh::SerializeRecastMeshTile(FArchive& Ar, int32 NavMeshVer
 		// sort out where various data types do/will live
 		unsigned char* d = TileData;
 		dtMeshHeader* Header = (dtMeshHeader*)d; d += TileLayout.HeaderSize;
-		FVector::FReal* NavVerts = (FVector::FReal*)d; d += TileLayout.VertsSize;
+		dtTileVert* NavVerts = (dtTileVert*)d; d += TileLayout.VertsSize;
 		dtPoly* NavPolys = (dtPoly*)d; d += TileLayout.PolysSize;
 		d += TileLayout.LinksSize;			// @fixme, are links autogenerated on addTile?
 		dtPolyDetail* DetailMeshes = (dtPolyDetail*)d; d += TileLayout.DetailMeshesSize;
-		FVector::FReal* DetailVerts = (FVector::FReal*)d; d += TileLayout.DetailVertsSize;
+		dtTileVert* DetailVerts = (dtTileVert*)d; d += TileLayout.DetailVertsSize;
 		unsigned char* DetailTris = (unsigned char*)d; d += TileLayout.DetailTrisSize;
 		dtBVNode* BVTree = (dtBVNode*)d; d += TileLayout.BvTreeSize;
 		dtOffMeshConnection* OffMeshCons = (dtOffMeshConnection*)d; d += TileLayout.OffMeshConsSize;
@@ -902,14 +906,50 @@ void FPImplRecastNavMesh::SerializeRecastMeshTile(FArchive& Ar, int32 NavMeshVer
 		Ar << DummySegmentInt << DummySegmentInt << DummySegmentInt;
 #endif // WITH_NAVMESH_SEGMENT_LINKS
 
-		// mesh and offmesh connection vertices, just an array of reals (one real triplet per vert)
+		// mesh and offmesh connection vertices
+		if (NavMeshVersion < NAVMESHVER_TILEVERT_ENCODING)
+		{
+			// Old format is in world-space (FVector::FReal * 3 per vertex).
+			if (Ar.IsLoading())
+			{
+				// Load old format and convert to tile-local dtTileVert.
+				for (int32 VertIdx = 0; VertIdx < SizeInfo.VertCount; ++VertIdx)
+				{
+					FVector::FReal wx, wy, wz;
+					Ar << wx;
+					Ar << wy;
+					Ar << wz;
+					const dtReal worldPos[3] = {(dtReal)wx, (dtReal)wy, (dtReal)wz};
+					NavVerts[VertIdx].setFromWorld(Header->bmin, worldPos);
+				}
+
+				// Update navigation tile data version to latest so it becomes readable.
+				Header->version = DT_NAVMESH_VERSION;
+			}
+			else
+			{
+				// Save in old format: convert from tile-local dtTileVert to world-space.
+				for (int32 VertIdx = 0; VertIdx < SizeInfo.VertCount; ++VertIdx)
+				{
+					dtReal worldPos[3];
+					NavVerts[VertIdx].toWorld(Header->bmin, worldPos);
+					FVector::FReal wx = (FVector::FReal)worldPos[0];
+					FVector::FReal wy = (FVector::FReal)worldPos[1];
+					FVector::FReal wz = (FVector::FReal)worldPos[2];
+					Ar << wx;
+					Ar << wy;
+					Ar << wz;
+				}
+			}
+		}
+		else
 		{
-			FVector::FReal* F = NavVerts;
-			for (int32 VertIdx=0; VertIdx < SizeInfo.VertCount; VertIdx++)
+			// New format: tile-local floats stored directly in dtTileVert.
+			for (int32 VertIdx = 0; VertIdx < SizeInfo.VertCount; ++VertIdx)
 			{
-				Ar << *F; F++;
-				Ar << *F; F++;
-				Ar << *F; F++;
+				Ar << NavVerts[VertIdx].x;
+				Ar << NavVerts[VertIdx].y;
+				Ar << NavVerts[VertIdx].z;
 			}
 		}
 
@@ -937,14 +977,47 @@ void FPImplRecastNavMesh::SerializeRecastMeshTile(FArchive& Ar, int32 NavMeshVer
 			Ar << DM.vertBase << DM.triBase << DM.vertCount << DM.triCount;
 		}
 
-		// serialize detail verts (one real triplet per vert)
+		// serialize detail verts
+		if (NavMeshVersion < NAVMESHVER_TILEVERT_ENCODING)
 		{
-			FVector::FReal* F = DetailVerts;
-			for (int32 VertIdx=0; VertIdx < SizeInfo.DetailVertCount; ++VertIdx)
+			// Old format is in world-space (FVector::FReal * 3 per vertex).
+			if (Ar.IsLoading())
 			{
-				Ar << *F; F++;
-				Ar << *F; F++;
-				Ar << *F; F++;
+				// Load old format and convert to tile-local dtTileVert.
+				for (int32 VertIdx = 0; VertIdx < SizeInfo.DetailVertCount; ++VertIdx)
+				{
+					FVector::FReal wx, wy, wz;
+					Ar << wx;
+					Ar << wy;
+					Ar << wz;
+					const dtReal worldPos[3] = {(dtReal)wx, (dtReal)wy, (dtReal)wz};
+					DetailVerts[VertIdx].setFromWorld(Header->bmin, worldPos);
+				}
+			}
+			else
+			{
+				// Save in old format: convert from tile-local dtTileVert to world-space.
+				for (int32 VertIdx = 0; VertIdx < SizeInfo.DetailVertCount; ++VertIdx)
+				{
+					dtReal worldPos[3];
+					DetailVerts[VertIdx].toWorld(Header->bmin, worldPos);
+					FVector::FReal wx = (FVector::FReal)worldPos[0];
+					FVector::FReal wy = (FVector::FReal)worldPos[1];
+					FVector::FReal wz = (FVector::FReal)worldPos[2];
+					Ar << wx;
+					Ar << wy;
+					Ar << wz;
+				}
+			}
+		}
+		else
+		{
+			// New format: tile-local floats stored directly in dtTileVert.
+			for (int32 VertIdx = 0; VertIdx < SizeInfo.DetailVertCount; ++VertIdx)
+			{
+				Ar << DetailVerts[VertIdx].x;
+				Ar << DetailVerts[VertIdx].y;
+				Ar << DetailVerts[VertIdx].z;
 			}
 		}
 
@@ -1061,7 +1134,7 @@ void FPImplRecastNavMesh::SerializeCompressedTileCacheData(FArchive& Ar, int32 N
 
 	if (Ar.IsLoading())
 	{
-		// Initialize to 0 if we are loading as this is calculated and used duiring processing.
+		// Initialize to 0 if we are loading as this is calculated and used during processing.
 		CompressedDataSize = 0;
 	}
 	
@@ -1073,7 +1146,7 @@ void FPImplRecastNavMesh::SerializeCompressedTileCacheData(FArchive& Ar, int32 N
 	if (Ar.IsSaving())
 	{
 		// Handle invalid CompressedDataSize ( i.e. CompressedDataSize > 0 and CompressedDataSize < dtAlign(sizeof(dtTileCacheLayerHeader))
-		// as well as valid CompressedDataSize == 0, to make this function atleast as robust as it was.
+		// as well as valid CompressedDataSize == 0, to make this function at least as robust as it was.
 		if (CompressedDataSize < dtAlign(sizeof(dtTileCacheLayerHeader)))
 		{
 			CompressedDataSizeNoHeader = EmptyDataValue;
@@ -1101,7 +1174,10 @@ void FPImplRecastNavMesh::SerializeCompressedTileCacheData(FArchive& Ar, int32 N
 		{
 			UE_LOG(LogNavigation, Error, TEXT("Failed to alloc tile compressed data"));
 		}
-		FMemory::Memset(CompressedData, 0, CompressedDataSize);
+		else
+		{
+			FMemory::Memset(CompressedData, 0, CompressedDataSize);
+		}
 	}
 
 	check(CompressedData != nullptr);
@@ -1676,7 +1752,9 @@ static void StorePathfindingDebugData(const dtNavMeshQuery& NavQuery, const dtNa
 			NodeInfo.NumVerts = NavPoly->vertCount;
 			for (int32 VertIdx = 0; VertIdx < NavPoly->vertCount; VertIdx++)
 			{
-				NodeInfo.Verts.Add((FVector3f)Recast2UnrealPoint(&NavTile->verts[NavPoly->verts[VertIdx] * 3]));
+				dtReal V[3];
+				NavTile->verts[NavPoly->verts[VertIdx]].toWorld(NavTile->header->bmin, V);
+				NodeInfo.Verts.Add((FVector3f)Recast2UnrealPoint(V));
 			}
 		}
 
@@ -2190,7 +2268,8 @@ bool FPImplRecastNavMesh::GetPolyCenter(NavNodeRef PolyID, FVector& OutCenter) c
 
 			for (uint32 VertIdx=0; VertIdx < Poly->vertCount; ++VertIdx)
 			{
-				const FVector::FReal* V = &Tile->verts[Poly->verts[VertIdx]*3];
+				FVector::FReal V[3];
+				Tile->verts[Poly->verts[VertIdx]].toWorld(Tile->header->bmin, V);
 				Center[0] += V[0];
 				Center[1] += V[1];
 				Center[2] += V[2];
@@ -2226,7 +2305,8 @@ bool FPImplRecastNavMesh::GetPolyVerts(NavNodeRef PolyID, TArray<FVector>& OutVe
 			// convert to UE coords and copy verts into output array 
 			for (uint32 VertIdx=0; VertIdx < Poly->vertCount; ++VertIdx)
 			{
-				const FVector::FReal* V = &Tile->verts[Poly->verts[VertIdx]*3];
+				FVector::FReal V[3];
+				Tile->verts[Poly->verts[VertIdx]].toWorld(Tile->header->bmin, V);
 				OutVerts.Add( Recast2UnrVector(V) );
 			}
 
@@ -2278,9 +2358,10 @@ FVector::FReal FPImplRecastNavMesh::GetPolySurfaceArea(NavNodeRef PolyID) const
 			dtReal PolyArea = 0;
 			for (int j = 2; j < Poly->vertCount; ++j)
 			{
-				const dtReal* VA = &Tile->verts[Poly->verts[0] * 3];
-				const dtReal* VB = &Tile->verts[Poly->verts[j - 1] * 3];
-				const dtReal* VC = &Tile->verts[Poly->verts[j] * 3];
+				dtReal VA[3], VB[3], VC[3];
+				Tile->verts[Poly->verts[0]].toWorld(Tile->header->bmin, VA);
+				Tile->verts[Poly->verts[j - 1]].toWorld(Tile->header->bmin, VB);
+				Tile->verts[Poly->verts[j]].toWorld(Tile->header->bmin, VC);
 				PolyArea += dtTriArea2D(VA, VB, VC);
 			}
 
@@ -2427,8 +2508,11 @@ bool FPImplRecastNavMesh::GetPolyEdges(NavNodeRef PolyID, TArray<FNavigationPort
 				if (LinkInfo.edge >= 0 && LinkInfo.edge < Poly->vertCount)
 				{
 					FNavigationPortalEdge NeiData;
-					NeiData.Left = Recast2UnrealPoint(&Tile->verts[3 * Poly->verts[LinkInfo.edge]]);
-					NeiData.Right = bIsNavLink ? NeiData.Left : Recast2UnrealPoint(&Tile->verts[3 * Poly->verts[(LinkInfo.edge + 1) % Poly->vertCount]]);
+					dtReal VLeft[3], VRight[3];
+					Tile->verts[Poly->verts[LinkInfo.edge]].toWorld(Tile->header->bmin, VLeft);
+					Tile->verts[Poly->verts[(LinkInfo.edge + 1) % Poly->vertCount]].toWorld(Tile->header->bmin, VRight);
+					NeiData.Left = Recast2UnrealPoint(VLeft);
+					NeiData.Right = bIsNavLink ? NeiData.Left : Recast2UnrealPoint(VRight);
 					NeiData.ToRef = LinkInfo.ref;
 					Edges.Add(NeiData);
 				}
@@ -2594,7 +2678,8 @@ bool FPImplRecastNavMesh::GetClusterBounds(NavNodeRef ClusterRef, FBox& OutBound
 				const dtPoly* Poly = &Tile->polys[i];
 				for (int32 iVert = 0; iVert < Poly->vertCount; iVert++)
 				{
-					const FVector::FReal* V = &Tile->verts[Poly->verts[iVert]*3];
+					FVector::FReal V[3];
+					Tile->verts[Poly->verts[iVert]].toWorld(Tile->header->bmin, V);
 					OutBounds += Recast2UnrealPoint(V);
 				}
 
@@ -2690,7 +2775,9 @@ bool FPImplRecastNavMesh::GetPolysInTile(int32 TileIndex, TArray<FNavPoly>& Poly
 			FVector PolyCenter(0);
 			for (int k = 0; k < Poly->vertCount; ++k)
 			{
-				PolyCenter += Recast2UnrealPoint(&Tile->verts[Poly->verts[k]*3]);
+				dtReal V[3];
+				Tile->verts[Poly->verts[k]].toWorld(Tile->header->bmin, V);
+				PolyCenter += Recast2UnrealPoint(V);
 			}
 			PolyCenter /= Poly->vertCount;
 
@@ -2776,8 +2863,9 @@ void FPImplRecastNavMesh::GetTilePolyEdgesForFilter(const dtMeshTile& Tile, cons
 								if (!IsPolyExterior(*Polyj, *PolyjTile, Link.ref, DetourNavMesh->getPolyRefBase(PolyjTile)))
 								{
 									// Poly i is interior and neighbor j is exterior
-									const FVector::FReal* V0 = &Tile.verts[Poly->verts[Link.edge] * 3];
-									const FVector::FReal* V1 = &Tile.verts[Poly->verts[(Link.edge + 1) % nj] * 3];
+									FVector::FReal V0[3], V1[3];
+									Tile.verts[Poly->verts[Link.edge]].toWorld(Tile.header->bmin, V0);
+									Tile.verts[Poly->verts[(Link.edge + 1) % nj]].toWorld(Tile.header->bmin, V1);
 									OutEdges.Emplace(Recast2UnrVector(V0), Recast2UnrVector(V1));
 								}
 							}
@@ -2794,8 +2882,9 @@ void FPImplRecastNavMesh::GetTilePolyEdgesForFilter(const dtMeshTile& Tile, cons
 					if (!bIsPolyjExterior)
 					{
 						// Poly i is interior and neighbor j is exterior
-						const FVector::FReal* V0 = &Tile.verts[Poly->verts[j] * 3];
-						const FVector::FReal* V1 = &Tile.verts[Poly->verts[(j + 1) % nj] * 3];
+						FVector::FReal V0[3], V1[3];
+						Tile.verts[Poly->verts[j]].toWorld(Tile.header->bmin, V0);
+						Tile.verts[Poly->verts[(j + 1) % nj]].toWorld(Tile.header->bmin, V1);
 						OutEdges.Emplace(Recast2UnrVector(V0), Recast2UnrVector(V1));
 					}
 				}
@@ -2810,8 +2899,9 @@ void FPImplRecastNavMesh::GetTilePolyEdgesForFilter(const dtMeshTile& Tile, cons
 				{
 					// Poly i is interior and neighbor j is exterior
 					// Swap the vertices to keep the same edges orientations
-					const FVector::FReal* V0 = &Tile.verts[Poly->verts[(j + 1) % nj] * 3];
-					const FVector::FReal* V1 = &Tile.verts[Poly->verts[j] * 3];
+					FVector::FReal V0[3], V1[3];
+					Tile.verts[Poly->verts[(j + 1) % nj]].toWorld(Tile.header->bmin, V0);
+					Tile.verts[Poly->verts[j]].toWorld(Tile.header->bmin, V1);
 					OutEdges.Emplace(Recast2UnrVector(V0), Recast2UnrVector(V1));
 				}
 			}
@@ -2879,25 +2969,29 @@ void FPImplRecastNavMesh::GetTilePolyEdges(const dtMeshTile& Tile, bool bGatherI
 				continue;
 			}
 
-			const FVector::FReal* V0 = &Tile.verts[Poly->verts[j] * 3];
-			const FVector::FReal* V1 = &Tile.verts[Poly->verts[(j + 1) % nj] * 3];
+			FVector::FReal V0[3], V1[3];
+			Tile.verts[Poly->verts[j]].toWorld(Tile.header->bmin, V0);
+			Tile.verts[Poly->verts[(j + 1) % nj]].toWorld(Tile.header->bmin, V1);
 
 			// Gather detail mesh edges which align with the actual poly edge.
 			// This is really slow.
 			for (int32 k = 0; k < pd->triCount; ++k)
 			{
 				const unsigned char* t = &(Tile.detailTris[(pd->triBase + k) * 4]);
+				FVector::FReal tvBuf[3][3];
 				const FVector::FReal* tv[3];
 
 				for (int32 m = 0; m < 3; ++m)
 				{
 					if (t[m] < Poly->vertCount)
 					{
-						tv[m] = &Tile.verts[Poly->verts[t[m]] * 3];
+						Tile.verts[Poly->verts[t[m]]].toWorld(Tile.header->bmin, tvBuf[m]);
+						tv[m] = tvBuf[m];
 					}
 					else
 					{
-						tv[m] = &Tile.detailVerts[(pd->vertBase + (t[m] - Poly->vertCount)) * 3];
+						Tile.detailVerts[pd->vertBase + (t[m] - Poly->vertCount)].toWorld(Tile.header->bmin, tvBuf[m]);
+						tv[m] = tvBuf[m];
 					}
 				}
 				for (int m = 0, n = 2; m < 3; n=m++)
@@ -3113,22 +3207,22 @@ int32 FPImplRecastNavMesh::GetTilesDebugGeometry(const FRecastNavMeshGenerator*
 	});
 
 	// add all the poly verts
-	FVector::FReal* F = Tile.verts;
 	for (int32 VertIdx = 0; VertIdx < Header->vertCount; ++VertIdx)
 	{
-		FVector const VertPos = Recast2UnrVector(F);
+		FVector::FReal V[3];
+		Tile.verts[VertIdx].toWorld(Header->bmin, V);
+		FVector const VertPos = Recast2UnrVector(V);
 		OutGeometry.MeshVerts.Add(VertPos);
-		F += 3;
 	}
 
 	int32 const DetailVertIndexBase = Header->vertCount;
 	// add the detail verts
-	F = Tile.detailVerts;
 	for (int32 DetailVertIdx = 0; DetailVertIdx < Header->detailVertCount; ++DetailVertIdx)
 	{
-		FVector const VertPos = Recast2UnrVector(F);
+		FVector::FReal V[3];
+		Tile.detailVerts[DetailVertIdx].toWorld(Header->bmin, V);
+		FVector const VertPos = Recast2UnrVector(V);
 		OutGeometry.MeshVerts.Add(VertPos);
-		F += 3;
 	}
 
 #if RECAST_INTERNAL_DEBUG_DATA	
@@ -3225,8 +3319,9 @@ int32 FPImplRecastNavMesh::GetTilesDebugGeometry(const FRecastNavMeshGenerator*
 		if (OffMeshConnection != NULL)
 		{
 			dtPoly const* const LinkPoly = &Tile.polys[OffMeshConnection->poly];
-			const FVector::FReal* va = &Tile.verts[LinkPoly->verts[0] * 3]; //OffMeshConnection->pos;
-			const FVector::FReal* vb = &Tile.verts[LinkPoly->verts[1] * 3]; //OffMeshConnection->pos[3];
+			FVector::FReal va[3], vb[3];
+			Tile.verts[LinkPoly->verts[0]].toWorld(Tile.header->bmin, va); //OffMeshConnection->pos;
+			Tile.verts[LinkPoly->verts[1]].toWorld(Tile.header->bmin, vb); //OffMeshConnection->pos[3];
 
 			const FRecastDebugGeometry::FOffMeshLink Link = {
 				Recast2UnrVector(va)
@@ -3259,10 +3354,15 @@ int32 FPImplRecastNavMesh::GetTilesDebugGeometry(const FRecastNavMeshGenerator*
 				dtPoly const* const LinkPoly = &Tile.polys[polyBase + j];
 
 				FRecastDebugGeometry::FOffMeshSegment Link;
-				Link.LeftStart = Recast2UnrealPoint(&Tile.verts[LinkPoly->verts[0] * 3]);
-				Link.LeftEnd = Recast2UnrealPoint(&Tile.verts[LinkPoly->verts[1] * 3]);
-				Link.RightStart = Recast2UnrealPoint(&Tile.verts[LinkPoly->verts[2] * 3]);
-				Link.RightEnd = Recast2UnrealPoint(&Tile.verts[LinkPoly->verts[3] * 3]);
+				FVector::FReal vLS[3], vLE[3], vRS[3], vRE[3];
+				Tile.verts[LinkPoly->verts[0]].toWorld(Tile.header->bmin, vLS);
+				Tile.verts[LinkPoly->verts[1]].toWorld(Tile.header->bmin, vLE);
+				Tile.verts[LinkPoly->verts[2]].toWorld(Tile.header->bmin, vRS);
+				Tile.verts[LinkPoly->verts[3]].toWorld(Tile.header->bmin, vRE);
+				Link.LeftStart = Recast2UnrealPoint(vLS);
+				Link.LeftEnd = Recast2UnrealPoint(vLE);
+				Link.RightStart = Recast2UnrealPoint(vRS);
+				Link.RightEnd = Recast2UnrealPoint(vRE);
 				Link.AreaID = LinkPoly->getArea();
 				Link.Direction = (uint8)OffMeshSeg->getBiDirectional();
 				Link.ValidEnds = GetValidEnds(*DetourNavMesh, Tile, *LinkPoly);
@@ -3621,4 +3721,4 @@ bool FPImplRecastNavMesh::HasTileCacheLayers(int32 TileX, int32 TileY) const
 
 #undef INITIALIZE_NAVQUERY
 
-#endif // WITH_RECAST
\ No newline at end of file
+#endif // WITH_RECAST
diff --git a/Engine/Source/Runtime/NavigationSystem/Private/NavMesh/RecastNavMesh.cpp b/Engine/Source/Runtime/NavigationSystem/Private/NavMesh/RecastNavMesh.cpp
--- a/Engine/Source/Runtime/NavigationSystem/Private/NavMesh/RecastNavMesh.cpp
+++ b/Engine/Source/Runtime/NavigationSystem/Private/NavMesh/RecastNavMesh.cpp
@@ -146,11 +146,11 @@ void FDetourTileLayout::InitFromSizeInfo(const FDetourTileSizeInfo& SizeInfo)
 #if WITH_RECAST
 	// Patch header pointers.
 	HeaderSize = dtAlign(sizeof(dtMeshHeader));
-	VertsSize = dtAlign(sizeof(dtReal) * 3 * SizeInfo.VertCount);
+	VertsSize = dtAlign(sizeof(dtTileVert) * SizeInfo.VertCount);
 	PolysSize = dtAlign(sizeof(dtPoly) * SizeInfo.PolyCount);
 	LinksSize = dtAlign(sizeof(dtLink) * (SizeInfo.MaxLinkCount));
 	DetailMeshesSize = dtAlign(sizeof(dtPolyDetail) * SizeInfo.DetailMeshCount);
-	DetailVertsSize = dtAlign(sizeof(dtReal) * 3 * SizeInfo.DetailVertCount);
+	DetailVertsSize = dtAlign(sizeof(dtTileVert) * SizeInfo.DetailVertCount);
 	DetailTrisSize = dtAlign(sizeof(unsigned char) * 4 * SizeInfo.DetailTriCount);
 	BvTreeSize = dtAlign(sizeof(dtBVNode) * SizeInfo.BvNodeCount);
 	OffMeshConsSize = dtAlign(sizeof(dtOffMeshConnection) * SizeInfo.OffMeshConCount);
@@ -265,6 +265,31 @@ static FAutoConsoleCommandWithWorldAndArgs CmdNavMeshDrawDistance(
 		}
 	));
 
+#if !UE_BUILD_SHIPPING
+static FAutoConsoleCommandWithWorldArgsAndOutputDevice CmdDumpNavMeshMemory(
+	TEXT("ai.debug.nav.DumpNavMeshMemory"),
+	TEXT("Dumps detailed tile memory usage statistics for all active navmeshes."),
+	FConsoleCommandWithWorldArgsAndOutputDeviceDelegate::CreateLambda(
+		[](const TArray<FString>& Args, const UWorld* World, FOutputDevice& OutputDevice)
+		{
+			if (const UNavigationSystemV1* NavSys = FNavigationSystem::GetCurrent<UNavigationSystemV1>(World))
+			{
+				for (const ANavigationData* NavData : NavSys->NavDataSet)
+				{
+					if (const ARecastNavMesh* RecastNavMesh = Cast<ARecastNavMesh>(NavData))
+					{
+						RecastNavMesh->DumpDetailedMemoryStats();
+					}
+				}
+			}
+			else
+			{
+				OutputDevice.Log(ELogVerbosity::Warning, TEXT("Command failed: unable to find the navigation system."));
+			}
+		}
+	));
+#endif // !UE_BUILD_SHIPPING
+
 FVector::FReal ARecastNavMesh::DrawDistanceSq = 0.;
 float ARecastNavMesh::MinimumSizeForChaosNavMeshInfluenceSq = 0.0f;
 #if !WITH_RECAST
@@ -1270,10 +1295,11 @@ void ARecastNavMesh::SerializeRecastNavMesh(FArchive& Ar, FPImplRecastNavMesh*&
 
 void ARecastNavMesh::Serialize( FArchive& Ar )
 {
-	UE_LOG(LogNavigation, Verbose, TEXT("%hs %s '%s'"), __FUNCTION__,
+	UE_LOG(LogNavigation, Verbose, TEXT("%hs %s %s, '%s'"), __FUNCTION__,
 		Ar.IsLoading() ? TEXT("'loading'") :
 			(Ar.IsSaving() ? TEXT("'saving'") :
 				(Ar.IsTransacting() ? TEXT("'transacting'") : TEXT("'unknown'"))),
+		*Ar.GetArchiveName(),
 		*GetFullNameSafe(this));
 
 	Super::Serialize(Ar);
@@ -1318,6 +1344,11 @@ void ARecastNavMesh::Serialize( FArchive& Ar )
 		else if (RecastNavMeshSizeBytes > 4)
 		{
 			SerializeRecastNavMesh(Ar, RecastNavMeshImpl, NavMeshVersion);
+			if (Ar.IsError())
+			{
+				return;
+			}
+
 #if !(UE_BUILD_SHIPPING)
 			RequestDrawingUpdate();
 #endif //!(UE_BUILD_SHIPPING)
@@ -1336,6 +1367,10 @@ void ARecastNavMesh::Serialize( FArchive& Ar )
 	else
 	{
 		SerializeRecastNavMesh(Ar, RecastNavMeshImpl, NavMeshVersion);
+		if (Ar.IsError())
+		{
+			return;
+		}
 
 		if (Ar.IsSaving())
 		{
@@ -1651,7 +1686,9 @@ bool ARecastNavMesh::GetPolysInTile(FNavTileRef TileRef, TArray<FNavPoly>& Polys
 				FVector PolyCenter(0);
 				for (int k = 0; k < Poly->vertCount; ++k)
 				{
-					PolyCenter += Recast2UnrealPoint(&Tile->verts[Poly->verts[k] * 3]);
+					FVector::FReal Vertex[3];
+					Tile->verts[Poly->verts[k]].toWorld(Tile->header->bmin, Vertex);
+					PolyCenter += Recast2UnrealPoint(Vertex);
 				}
 				PolyCenter /= Poly->vertCount;
 
@@ -1694,7 +1731,11 @@ bool ARecastNavMesh::GetNavLinksInTile(const int32 TileIndex, TArray<FNavPoly>&
 				FNavPoly& OutPoly = Polys[BaseIdx + LinkIndex];
 				const int32 PolyIndex = Tile->header->offMeshBase + LinkIndex;
 				OutPoly.Ref = DetourNavMesh->encodePolyId(Tile->salt, TileIndex, PolyIndex);
-				OutPoly.Center = (Recast2UnrealPoint(&Tile->verts[Poly->verts[0] * 3]) + Recast2UnrealPoint(&Tile->verts[Poly->verts[1] * 3])) / 2;
+				FVector::FReal VertexA[3];
+				FVector::FReal VertexB[3];
+				Tile->verts[Poly->verts[0]].toWorld(Tile->header->bmin, VertexA);
+				Tile->verts[Poly->verts[1]].toWorld(Tile->header->bmin, VertexB);
+				OutPoly.Center = (Recast2UnrealPoint(VertexA) + Recast2UnrealPoint(VertexB)) / 2;
 			}
 		}
 
@@ -2306,7 +2347,9 @@ bool ARecastNavMesh::GetPolysInBox(const FBox& Box, TArray<FNavPoly>& Polys, FSh
 					FVector PolyCenter(0);
 					for (int k = 0; k < Poly->vertCount; ++k)
 					{
-						PolyCenter += Recast2UnrealPoint(&Tile->verts[Poly->verts[k]*3]);
+						FVector::FReal Vertex[3];
+						Tile->verts[Poly->verts[k]].toWorld(Tile->header->bmin, Vertex);
+						PolyCenter += Recast2UnrealPoint(Vertex);
 					}
 					PolyCenter /= Poly->vertCount;
 
@@ -3213,6 +3256,229 @@ uint32 ARecastNavMesh::LogMemUsed() const
 	return MemUsed + SuperMemUsed;
 }
 
+void ARecastNavMesh::DumpDetailedMemoryStats() const
+{
+#if WITH_RECAST
+	if (!RecastNavMeshImpl || !RecastNavMeshImpl->DetourNavMesh)
+	{
+		UE_LOGF(LogNavigation, Display, "%ls: No Detour navmesh data.", *GetName());
+		return;
+	}
+
+	const dtNavMesh* NavMesh = RecastNavMeshImpl->DetourNavMesh;
+	const int32 MaxTiles = NavMesh->getMaxTiles();
+
+	// Aggregate counters
+	int32 TotalActiveTiles = 0;
+	uint64 TotalTileDataBytes = 0;
+
+	// Per-section aggregates
+	uint64 TotalHeaderBytes = 0;
+	uint64 TotalVertBytes = 0;
+	uint64 TotalPolyBytes = 0;
+	uint64 TotalLinkBytes = 0;
+	uint64 TotalDetailMeshBytes = 0;
+	uint64 TotalDetailVertBytes = 0;
+	uint64 TotalDetailTriBytes = 0;
+	uint64 TotalBvTreeBytes = 0;
+	uint64 TotalOffMeshConBytes = 0;
+	uint64 TotalOffMeshSegBytes = 0;
+	uint64 TotalClusterBytes = 0;
+	uint64 TotalPolyClusterBytes = 0;
+
+	// Count trackers
+	int32 TilesWithBvTree = 0;
+	int32 TilesWithoutBvTree = 0;
+	int32 TotalPolyCount = 0;
+	int32 TotalVertCount = 0;
+	int32 TotalLinkCount = 0;
+	int32 TotalBvNodeCount = 0;
+	int32 TotalDetailVertCount = 0;
+	int32 TotalDetailTriCount = 0;
+	int32 TotalOffMeshConCount = 0;
+	int32 MinPolysPerTile = MAX_int32;
+	int32 MaxPolysPerTile = 0;
+
+	// Histogram: poly count per tile
+	// Buckets: [1-2], [3-4], [5-8], [9-16], [17-32], [33-64], [65-128], [129-256], [257+]
+	static constexpr int32 NumBuckets = 9;
+	static constexpr int32 BucketMaxPolys[NumBuckets] = { 2, 4, 8, 16, 32, 64, 128, 256, MAX_int32 };
+	static const char* BucketLabels[NumBuckets] = {
+		"1-2", "3-4", "5-8", "9-16", "17-32", "33-64", "65-128", "129-256", "257+"
+	};
+	int32 BucketTileCount[NumBuckets] = {};
+	uint64 BucketTotalBytes[NumBuckets] = {};
+	uint64 BucketBvTreeBytes[NumBuckets] = {};
+	uint64 BucketDetailBytes[NumBuckets] = {};
+
+	// Iterate all tiles
+	for (int32 TileIndex = 0; TileIndex < MaxTiles; ++TileIndex)
+	{
+		const dtMeshTile* Tile = NavMesh->getTile(TileIndex);
+		if (!Tile || !Tile->header)
+		{
+			continue;
+		}
+
+		TotalActiveTiles++;
+		const dtMeshHeader* Header = Tile->header;
+		const FDetourTileLayout TileLayout(*Tile);
+
+		TotalTileDataBytes += TileLayout.TileSize;
+		TotalHeaderBytes += TileLayout.HeaderSize;
+		TotalVertBytes += TileLayout.VertsSize;
+		TotalPolyBytes += TileLayout.PolysSize;
+		TotalLinkBytes += TileLayout.LinksSize;
+		TotalDetailMeshBytes += TileLayout.DetailMeshesSize;
+		TotalDetailVertBytes += TileLayout.DetailVertsSize;
+		TotalDetailTriBytes += TileLayout.DetailTrisSize;
+		TotalBvTreeBytes += TileLayout.BvTreeSize;
+		TotalOffMeshConBytes += TileLayout.OffMeshConsSize;
+		TotalOffMeshSegBytes += TileLayout.OffMeshSegsSize;
+		TotalClusterBytes += TileLayout.ClustersSize;
+		TotalPolyClusterBytes += TileLayout.PolyClustersSize;
+
+		TotalPolyCount += Header->polyCount;
+		TotalVertCount += Header->vertCount;
+		TotalLinkCount += Header->maxLinkCount;
+		TotalBvNodeCount += Header->bvNodeCount;
+		TotalDetailVertCount += Header->detailVertCount;
+		TotalDetailTriCount += Header->detailTriCount;
+		TotalOffMeshConCount += Header->offMeshConCount;
+
+		if (Header->bvNodeCount > 0)
+		{
+			TilesWithBvTree++;
+		}
+		else
+		{
+			TilesWithoutBvTree++;
+		}
+
+		MinPolysPerTile = FMath::Min(MinPolysPerTile, (int32)Header->polyCount);
+		MaxPolysPerTile = FMath::Max(MaxPolysPerTile, (int32)Header->polyCount);
+
+		// Bucket by poly count
+		for (int32 Bucket = 0; Bucket < NumBuckets; ++Bucket)
+		{
+			if ((int32)Header->polyCount <= BucketMaxPolys[Bucket])
+			{
+				BucketTileCount[Bucket]++;
+				BucketTotalBytes[Bucket] += TileLayout.TileSize;
+				BucketBvTreeBytes[Bucket] += TileLayout.BvTreeSize;
+				BucketDetailBytes[Bucket] += TileLayout.DetailMeshesSize
+					+ TileLayout.DetailVertsSize + TileLayout.DetailTrisSize;
+				break;
+			}
+		}
+	}
+
+	if (TotalActiveTiles == 0)
+	{
+		UE_LOGF(LogNavigation, Display, "%ls: No active tiles.", *GetName());
+		return;
+	}
+
+	// Tile pool overhead
+	const uint64 TilePoolBytes = (uint64)MaxTiles * sizeof(dtMeshTile);
+	const uint64 LookupTableBytes =
+		(uint64)FMath::Max(1u, dtNextPow2(static_cast<unsigned int>(MaxTiles) / 4)) * sizeof(dtMeshTile*);
+
+	// Log everything
+	constexpr float ToKB = 1.0f / 1024.0f;
+	constexpr float ToMB = 1.0f / (1024.0f * 1024.0f);
+
+	UE_LOGF(LogNavigation, Display, "=======================================================================================================");
+	UE_LOGF(LogNavigation, Display, "  NavMesh Memory Report: %ls", *GetName());
+	UE_LOGF(LogNavigation, Display, "=======================================================================================================");
+	UE_LOGF(LogNavigation, Display, "  Max tiles (pool):    %d", MaxTiles);
+	UE_LOGF(LogNavigation, Display, "  Active tiles:        %d", TotalActiveTiles);
+	UE_LOGF(LogNavigation, Display, "  Pool utilization:    %.1f%%",
+		100.0f * (float)TotalActiveTiles / FMath::Max(MaxTiles, 1.f));
+	UE_LOGF(LogNavigation, Display, "");
+
+	// Overall memory
+	UE_LOGF(LogNavigation, Display, "--- Overall Memory ---");
+	UE_LOGF(LogNavigation, Display, "  Tile data total:     %10llu bytes  (%.2f MB)", TotalTileDataBytes, (float)TotalTileDataBytes * ToMB);
+	UE_LOGF(LogNavigation, Display, "  Tile pool structs:   %10llu bytes  (%.2f MB)", TilePoolBytes, (float)TilePoolBytes * ToMB);
+	UE_LOGF(LogNavigation, Display, "  Lookup table:        %10llu bytes  (%.2f KB)", LookupTableBytes, (float)LookupTableBytes * ToKB);
+	UE_LOGF(LogNavigation, Display, "  TOTAL:               %10llu bytes  (%.2f MB)",
+		(TotalTileDataBytes + TilePoolBytes + LookupTableBytes),
+		(float)(TotalTileDataBytes + TilePoolBytes + LookupTableBytes) * ToMB);
+	UE_LOGF(LogNavigation, Display, "");
+
+	// Per-section breakdown
+	UE_LOGF(LogNavigation, Display, "--- Per-Section Breakdown (tile data only) ---");
+	UE_LOGF(LogNavigation, Display, "  %-20s %9s %7s %6s    %s",
+		"Section", "Bytes", "MB", "%", "Description");
+
+	auto LogSection = [&](const char* Name, uint64 Bytes, const char* Desc)
+	{
+		const float Pct = TotalTileDataBytes > 0
+			? 100.0f * (float)Bytes / (float)TotalTileDataBytes : 0.0f;
+
+		UE_LOGF(LogNavigation, Display, "  %-20s %9llu %7.2f %6.1f%%   %s",
+			Name, Bytes, (float)Bytes * ToMB, Pct, Desc);
+	};
+
+	LogSection("Header",        TotalHeaderBytes,      "");
+	LogSection("Vertices",      TotalVertBytes,
+		TCHAR_TO_ANSI(*FString::Printf(TEXT("%d verts, %d bytes/vert"), TotalVertCount, TotalVertCount > 0 ? (int32)(TotalVertBytes / TotalVertCount) : 0)));
+	LogSection("Polygons",      TotalPolyBytes, TCHAR_TO_ANSI(*FString::Printf(TEXT("%d polys"), TotalPolyCount)));
+	LogSection("Links",         TotalLinkBytes, TCHAR_TO_ANSI(*FString::Printf(TEXT("%d max links allocated"), TotalLinkCount)));
+	LogSection("Detail Meshes", TotalDetailMeshBytes, "");
+	LogSection("Detail Verts",  TotalDetailVertBytes, TCHAR_TO_ANSI(*FString::Printf(TEXT("%d detail verts"), TotalDetailVertCount)));
+	LogSection("Detail Tris",   TotalDetailTriBytes, TCHAR_TO_ANSI(*FString::Printf(TEXT("%d detail tris"), TotalDetailTriCount)));
+	LogSection("BV Tree",       TotalBvTreeBytes, TCHAR_TO_ANSI(*FString::Printf(TEXT("%d nodes, %d/%d tiles have BV tree"), TotalBvNodeCount, TilesWithBvTree, TotalActiveTiles)));
+	LogSection("OffMesh Cons",  TotalOffMeshConBytes, TCHAR_TO_ANSI(*FString::Printf(TEXT("%d connections"), TotalOffMeshConCount)));
+	LogSection("OffMesh Segs",  TotalOffMeshSegBytes, "");
+	LogSection("Clusters",      TotalClusterBytes, "");
+	LogSection("PolyClusters",  TotalPolyClusterBytes, "");
+	UE_LOGF(LogNavigation, Display, "");
+
+	// Per-tile statstistics
+	UE_LOGF(LogNavigation, Display, "--- Per-Tile Statistics ---");
+	const uint64 AvgTileBytes = TotalTileDataBytes / TotalActiveTiles;
+	UE_LOGF(LogNavigation, Display, "  Avg tile data size:  %llu bytes (%.4f MB)", AvgTileBytes, (float)AvgTileBytes * ToMB);
+	UE_LOGF(LogNavigation, Display, "  Min polys/tile:      %d", MinPolysPerTile);
+	UE_LOGF(LogNavigation, Display, "  Max polys/tile:      %d", MaxPolysPerTile);
+	UE_LOGF(LogNavigation, Display, "  Avg polys/tile:      %.1f", (float)TotalPolyCount / (float)TotalActiveTiles);
+	UE_LOGF(LogNavigation, Display, "");
+
+	// Poly count histogram
+	UE_LOGF(LogNavigation, Display, "--- Tile Distribution by Poly Count ---");
+	UE_LOGF(LogNavigation, Display, "  %-16s %6s %6s     %9s  %10s  %10s  %10s   %10s  %10s",
+		"Polys", "Tiles", "%",
+		"Total KB", "Total MB", "BvTree KB", "BvTree MB", "Detail KB", "Detail MB");
+
+	for (int32 Bucket = 0; Bucket < NumBuckets; ++Bucket)
+	{
+		if (BucketTileCount[Bucket] == 0)
+		{
+			continue;
+		}
+		
+		const float TilePercentage = 100.0f * (float)BucketTileCount[Bucket] / (float)TotalActiveTiles;
+		UE_LOGF(LogNavigation, Display,
+			"  %-16hs %6d %6.1f%%     %9.1f  %10.2f  %10.1f  %10.2f   %10.1f  %10.2f",
+			BucketLabels[Bucket],
+			BucketTileCount[Bucket],
+			TilePercentage,
+			(float)BucketTotalBytes[Bucket] * ToKB,
+			(float)BucketTotalBytes[Bucket] * ToMB,
+			(float)BucketBvTreeBytes[Bucket] * ToKB,
+			(float)BucketBvTreeBytes[Bucket] * ToMB,
+			(float)BucketDetailBytes[Bucket] * ToKB,
+			(float)BucketDetailBytes[Bucket] * ToMB);
+	}
+
+	UE_LOGF(LogNavigation, Display, "=======================================================================================================");
+	UE_LOGF(LogNavigation, Display, "");
+
+#else  // WITH_RECAST
+	UE_LOGF(LogNavigation, Display, "%ls: Recast not available.", *GetName());
+#endif // WITH_RECAST
+}
 #endif // !UE_BUILD_SHIPPING
 
 uint16 ARecastNavMesh::GetDefaultForbiddenFlags() const
diff --git a/Engine/Source/Runtime/NavigationSystem/Private/NavMesh/RecastNavMeshGenerator.cpp b/Engine/Source/Runtime/NavigationSystem/Private/NavMesh/RecastNavMeshGenerator.cpp
--- a/Engine/Source/Runtime/NavigationSystem/Private/NavMesh/RecastNavMeshGenerator.cpp
+++ b/Engine/Source/Runtime/NavigationSystem/Private/NavMesh/RecastNavMeshGenerator.cpp
@@ -4339,7 +4339,7 @@ bool FRecastTileGenerator::GenerateNavigationDataLayer(FNavMeshBuildContext& Bui
 
 		// Skip BVtree for tiles having few polygons since linear scan is faster for low poly counts. 
 		// This saves sizeof(dtBVNode) * polyCount * 2 bytes per tile.
-		constexpr int32 PolyCountThreshold = 12;
+		constexpr int32 PolyCountThreshold = 16;
 		const bool bShouldBuildBvTree = TileConfig.bGenerateBVTree && (Params.polyCount > PolyCountThreshold);
 		Params.buildBvTree = bShouldBuildBvTree;
 
diff --git a/Engine/Source/Runtime/NavigationSystem/Private/NavigationSystem.cpp b/Engine/Source/Runtime/NavigationSystem/Private/NavigationSystem.cpp
--- a/Engine/Source/Runtime/NavigationSystem/Private/NavigationSystem.cpp
+++ b/Engine/Source/Runtime/NavigationSystem/Private/NavigationSystem.cpp
@@ -5897,7 +5897,7 @@ void UNavigationSystemV1::UpdateDelayedUnlockeRegistration(const TArray<FAssetCo
 
 void UNavigationSystemV1::RegisterDelayedUnlock(float DelayInSeconds)
 {
-	UE_LOG(LogNavigation, Verbose, TEXT("%hs"), __FUNCTION__);
+	UE_LOGF(LogNavigation, VeryVerbose, "%hs", __FUNCTION__);
 	
 	// Unregister the previous one and register a new one.
 	UnregisterDelayedUnlock();
@@ -5910,7 +5910,7 @@ void UNavigationSystemV1::RegisterDelayedUnlock(float DelayInSeconds)
 		0.f,
 		[this, FireInTickCount = 16, DelayInSeconds](float DeltaTime) mutable
 		{
-			UE_LOG(LogNavigation, Verbose, TEXT("%hs    Time remaining %f."), __FUNCTION__, DelayInSeconds);
+			UE_LOG(LogNavigation, VeryVerbose, TEXT("%hs    Time remaining %f."), __FUNCTION__, DelayInSeconds);
 				
 			DelayInSeconds -= DeltaTime;
 			if (--FireInTickCount <= 0 && DelayInSeconds <= 0.f && FAssetCompilingManager::Get().GetNumRemainingAssets() == 0)
diff --git a/Engine/Source/Runtime/NavigationSystem/Public/NavMesh/RecastNavMesh.h b/Engine/Source/Runtime/NavigationSystem/Public/NavMesh/RecastNavMesh.h
--- a/Engine/Source/Runtime/NavigationSystem/Public/NavMesh/RecastNavMesh.h
+++ b/Engine/Source/Runtime/NavigationSystem/Public/NavMesh/RecastNavMesh.h
@@ -1244,6 +1244,9 @@ public:
 
 #if !UE_BUILD_SHIPPING
 	NAVIGATIONSYSTEM_API virtual uint32 LogMemUsed() const override;
+
+	/** Dumps detailed per-section and per-tile memory statistics to the log. */
+	NAVIGATIONSYSTEM_API void DumpDetailedMemoryStats() const;	
 #endif // !UE_BUILD_SHIPPING
 
 	NAVIGATIONSYSTEM_API void UpdateNavMeshDrawing();
diff --git a/Engine/Source/Runtime/NavigationSystem/Public/NavMesh/RecastVersion.h b/Engine/Source/Runtime/NavigationSystem/Public/NavMesh/RecastVersion.h
--- a/Engine/Source/Runtime/NavigationSystem/Public/NavMesh/RecastVersion.h
+++ b/Engine/Source/Runtime/NavigationSystem/Public/NavMesh/RecastVersion.h
@@ -30,6 +30,7 @@
 #define NAVMESHVER_1_VOXEL_AGENT_STEEP_SLOPE_FILTER_FIX	25 // Fix, remove steep slope filtering during heightfield ledge filtering when the agent radius is included into a single voxel
 #define NAVMESHVER_TILE_RESOLUTIONS_AGENTMAXSTEPHEIGHT 26	// Addition of AgentMaxStepHeight in the resolution params, deprecating the original AgentMaxStepHeight.
 #define NAVMESHVER_NAVLINK_JUMP_CONFIGS			27 // Addition of NavLinkJumpConfigs, deprecating the original NavLinkJumpDownConfig
+#define NAVMESHVER_TILEVERT_ENCODING			28 // Store vertex data as tile-local floats (dtTileVert) instead of world-space doubles
 
-#define NAVMESHVER_LATEST				NAVMESHVER_NAVLINK_JUMP_CONFIGS
+#define NAVMESHVER_LATEST				NAVMESHVER_TILEVERT_ENCODING
 #define NAVMESHVER_MIN_COMPATIBLE		NAVMESHVER_LWCOORDS_OPTIMIZATION
diff --git a/Engine/Source/Runtime/Navmesh/Private/DebugUtils/DetourDebugDraw.cpp b/Engine/Source/Runtime/Navmesh/Private/DebugUtils/DetourDebugDraw.cpp
--- a/Engine/Source/Runtime/Navmesh/Private/DebugUtils/DetourDebugDraw.cpp
+++ b/Engine/Source/Runtime/Navmesh/Private/DebugUtils/DetourDebugDraw.cpp
@@ -95,21 +95,25 @@ static void drawPolyBoundaries(duDebugDraw* dd, const dtMeshTile* tile,
 				if (p->neis[j] != 0) continue;
 			}
 			
-			const dtReal* v0 = &tile->verts[p->verts[j]*3];
-			const dtReal* v1 = &tile->verts[p->verts[(j+1) % nj]*3];
-			
+			dtReal v0[3], v1[3]; //@UE
+			tile->verts[p->verts[j]].toWorld(tile->header->bmin, v0); //@UE
+			tile->verts[p->verts[(j+1) % nj]].toWorld(tile->header->bmin, v1); //@UE
+
 			// Draw detail mesh edges which align with the actual poly edge.
 			// This is really slow.
 			for (int k = 0; k < pd->triCount; ++k)
 			{
 				const unsigned char* t = &tile->detailTris[(pd->triBase+k)*4];
+				dtReal tvBuf[3][3]; //@UE
 				const dtReal* tv[3];
 				for (int m = 0; m < 3; ++m)
 				{
 					if (t[m] < p->vertCount)
-						tv[m] = &tile->verts[p->verts[t[m]]*3];
+						tile->verts[p->verts[t[m]]].toWorld(tile->header->bmin, tvBuf[m]); //@UE
 					else
-						tv[m] = &tile->detailVerts[(pd->vertBase+(t[m]-p->vertCount))*3];
+						tile->detailVerts[pd->vertBase+(t[m]-p->vertCount)].toWorld(tile->header->bmin, tvBuf[m]); //@UE
+
+					tv[m] = tvBuf[m]; //@UE
 				}
 				for (int m = 0, n = 2; m < 3; n=m++)
 				{
@@ -168,10 +172,13 @@ static void drawMeshTile(duDebugDraw* dd, const dtNavMesh& mesh, const dtNavMesh
 			const unsigned char* t = &tile->detailTris[(pd->triBase+j)*4];
 			for (int k = 0; k < 3; ++k)
 			{
+				dtReal tv[3]; //@UE
 				if (t[k] < p->vertCount)
-					dd->vertex(&tile->verts[p->verts[t[k]]*3], col);
+					tile->verts[p->verts[t[k]]].toWorld(tile->header->bmin, tv); //@UE
 				else
-					dd->vertex(&tile->detailVerts[(pd->vertBase+t[k]-p->vertCount)*3], col);
+					tile->detailVerts[pd->vertBase+t[k]-p->vertCount].toWorld(tile->header->bmin, tv); //@UE
+
+				dd->vertex(tv, col); //@UE
 			}
 		}
 	}
@@ -201,8 +208,9 @@ static void drawMeshTile(duDebugDraw* dd, const dtNavMesh& mesh, const dtNavMesh
 				col = duDarkenCol(duIntToCol(p->getArea(), 220));
 
 			const dtOffMeshConnection* con = &tile->offMeshCons[i - tile->header->offMeshBase];
-			const dtReal* va = &tile->verts[p->verts[0]*3];
-			const dtReal* vb = &tile->verts[p->verts[1]*3];
+			dtReal va[3], vb[3]; //@UE
+			tile->verts[p->verts[0]].toWorld(tile->header->bmin, va); //@UE
+			tile->verts[p->verts[1]].toWorld(tile->header->bmin, vb); //@UE
 
 			// Check to see if start and end end-points have links.
 			bool startSet = false;
@@ -251,10 +259,11 @@ static void drawMeshTile(duDebugDraw* dd, const dtNavMesh& mesh, const dtNavMesh
 			const dtPoly* p = &tile->polys[i];
 			if (p->getType() == DT_POLYTYPE_OFFMESH_SEGMENT)
 			{
-				dtReal* vA0 = &tile->verts[p->verts[0]*3];
-				dtReal* vA1 = &tile->verts[p->verts[1]*3];
-				dtReal* vB0 = &tile->verts[p->verts[2]*3];
-				dtReal* vB1 = &tile->verts[p->verts[3]*3];
+				dtReal vA0[3], vA1[3], vB0[3], vB1[3]; //@UE
+				tile->verts[p->verts[0]].toWorld(tile->header->bmin, vA0); //@UE
+				tile->verts[p->verts[1]].toWorld(tile->header->bmin, vA1); //@UE
+				tile->verts[p->verts[2]].toWorld(tile->header->bmin, vB0); //@UE
+				tile->verts[p->verts[3]].toWorld(tile->header->bmin, vB1); //@UE
 
 				unsigned int col = duDarkenCol(duIntToCol(p->getArea(), 220));
 				duAppendArc(dd, vA0[0], vA0[1], vA0[2], vB0[0], vB0[1], vB0[2], 0.25f, 0.0f, 0.0f, col);
@@ -276,10 +285,11 @@ static void drawMeshTile(duDebugDraw* dd, const dtNavMesh& mesh, const dtNavMesh
 			else
 				col = duDarkenCol(duIntToCol(p->getArea(), 220));
 
-			dtReal* vA0 = &tile->verts[p->verts[0]*3];
-			dtReal* vA1 = &tile->verts[p->verts[1]*3];
-			dtReal* vB0 = &tile->verts[p->verts[2]*3];
-			dtReal* vB1 = &tile->verts[p->verts[3]*3];
+			dtReal vA0[3], vA1[3], vB0[3], vB1[3]; //@UE
+			tile->verts[p->verts[0]].toWorld(tile->header->bmin, vA0); //@UE
+			tile->verts[p->verts[1]].toWorld(tile->header->bmin, vA1); //@UE
+			tile->verts[p->verts[2]].toWorld(tile->header->bmin, vB0); //@UE
+			tile->verts[p->verts[3]].toWorld(tile->header->bmin, vB1); //@UE
 			duAppendArcSegment(dd, vA0[0], vA0[1], vA0[2], vA1[0], vA1[1], vA1[2], vB0[0], vB0[1], vB0[2], vB1[0], vB1[1], vB1[2], 0.25f, col);
 		}
 		dd->end();
@@ -291,7 +301,8 @@ static void drawMeshTile(duDebugDraw* dd, const dtNavMesh& mesh, const dtNavMesh
 	dd->begin(DU_DRAW_POINTS, 3.0f);
 	for (int i = 0; i < tile->header->vertCount; ++i)
 	{
-		const dtReal* v = &tile->verts[i*3];
+		dtReal v[3]; //@UE
+		tile->verts[i].toWorld(tile->header->bmin, v); //@UE
 		dd->vertex(v[0], v[1], v[2], vcol);
 	}
 	dd->end();
@@ -424,8 +435,9 @@ static void drawMeshTilePortal(duDebugDraw* dd, const dtMeshTile* tile, const dt
 					continue;
 				
 				// Create new links
-				const dtReal* va = &tile->verts[poly->verts[j]*3];
-				const dtReal* vb = &tile->verts[poly->verts[(j+1) % nv]*3];
+				dtReal va[3], vb[3]; //@UE
+				tile->verts[poly->verts[j]].toWorld(tile->header->bmin, va); //@UE
+				tile->verts[poly->verts[(j+1) % nv]].toWorld(tile->header->bmin, vb); //@UE
 				
 				if (side == 0 || side == 4)
 				{
@@ -534,10 +546,11 @@ void duDebugDrawNavMeshPoly(duDebugDraw* dd, const dtNavMesh& mesh, dtPolyRef re
 	else if (poly->getType() == DT_POLYTYPE_OFFMESH_SEGMENT)
 	{
 		dd->begin(DU_DRAW_QUADS);
-		dtReal* vA0 = &tile->verts[poly->verts[0]*3];
-		dtReal* vA1 = &tile->verts[poly->verts[1]*3];
-		dtReal* vB0 = &tile->verts[poly->verts[2]*3];
-		dtReal* vB1 = &tile->verts[poly->verts[3]*3];
+		dtReal vA0[3], vA1[3], vB0[3], vB1[3];
+		tile->verts[poly->verts[0]].toWorld(tile->header->bmin, vA0);
+		tile->verts[poly->verts[1]].toWorld(tile->header->bmin, vA1);
+		tile->verts[poly->verts[2]].toWorld(tile->header->bmin, vB0);
+		tile->verts[poly->verts[3]].toWorld(tile->header->bmin, vB1);
 		duAppendArcSegment(dd, vA0[0], vA0[1], vA0[2], vA1[0], vA1[1], vA1[2], vB0[0], vB0[1], vB0[2], vB1[0], vB1[1], vB1[2], 0.25f, c);
 		dd->end();
 	}
@@ -553,10 +566,12 @@ void duDebugDrawNavMeshPoly(duDebugDraw* dd, const dtNavMesh& mesh, dtPolyRef re
 			const unsigned char* t = &tile->detailTris[(pd->triBase+i)*4];
 			for (int j = 0; j < 3; ++j)
 			{
+				dtReal tv[3]; //@UE
 				if (t[j] < poly->vertCount)
-					dd->vertex(&tile->verts[poly->verts[t[j]]*3], c);
+					tile->verts[poly->verts[t[j]]].toWorld(tile->header->bmin, tv); //@UE
 				else
-					dd->vertex(&tile->detailVerts[(pd->vertBase+t[j]-poly->vertCount)*3], c);
+					tile->detailVerts[pd->vertBase+t[j]-poly->vertCount].toWorld(tile->header->bmin, tv); //@UE
+				dd->vertex(tv, c); //@UE
 			}
 		}
 		dd->end();
diff --git a/Engine/Source/Runtime/Navmesh/Private/Detour/DetourNavMesh.cpp b/Engine/Source/Runtime/Navmesh/Private/Detour/DetourNavMesh.cpp
--- a/Engine/Source/Runtime/Navmesh/Private/Detour/DetourNavMesh.cpp
+++ b/Engine/Source/Runtime/Navmesh/Private/Detour/DetourNavMesh.cpp
@@ -176,17 +176,20 @@ static dtReal getHeightFromDMesh(const dtMeshTile* tile, int polyIdx, dtReal* po
 	for (int j = 0; j < pd->triCount; ++j)
 	{
 		const unsigned char* t = &tile->detailTris[(pd->triBase+j)*4];
+		dtReal vBuf[3][3]; //@UE
 		const dtReal* v[3];
 		for (int k = 0; k < 3; ++k)
 		{
 			if (t[k] < poly->vertCount)
-				v[k] = &tile->verts[poly->verts[t[k]]*3];
+				tile->verts[poly->verts[t[k]]].toWorld(tile->header->bmin, vBuf[k]); //@UE
 			else
-				v[k] = &tile->detailVerts[(pd->vertBase+(t[k]-poly->vertCount))*3];
+				tile->detailVerts[pd->vertBase+(t[k]-poly->vertCount)].toWorld(tile->header->bmin, vBuf[k]); //@UE
+
+			v[k] = vBuf[k]; //@UE
 		}
 		dtReal h;
 		if (dtClosestHeightPointTriangle(pos, v[0], v[1], v[2], h))
-		{			
+		{
 			return h;
 		}
 	}
@@ -391,12 +394,27 @@ static bool isPolyIntersectingSegment(const dtMeshTile* tile, int polyIdx,
 	const dtReal* spos, const dtReal* epos, dtReal& tmin, dtReal& tmax)
 {
 	dtPoly* poly = &tile->polys[polyIdx];
+
+	// Convert segment to local space instead of converting all vertices to world.
+	dtReal sposLocal[3], eposLocal[3]; //@UE
+	sposLocal[0] = spos[0] - tile->header->bmin[0]; //@UE
+	sposLocal[1] = spos[1] - tile->header->bmin[1]; //@UE
+	sposLocal[2] = spos[2] - tile->header->bmin[2]; //@UE
+	eposLocal[0] = epos[0] - tile->header->bmin[0]; //@UE
+	eposLocal[1] = epos[1] - tile->header->bmin[1]; //@UE
+	eposLocal[2] = epos[2] - tile->header->bmin[2]; //@UE
+
 	dtReal verts[DT_VERTS_PER_POLYGON*3];
 	for (int i = 0; i < poly->vertCount; i++)
-		dtVcopy(&verts[i*3], &tile->verts[poly->verts[i]*3]);
+	{
+		const dtTileVert& v = tile->verts[poly->verts[i]]; //@UE
+		verts[i*3 + 0] = (dtReal)v.x; //@UE
+		verts[i*3 + 1] = (dtReal)v.y; //@UE
+		verts[i*3 + 2] = (dtReal)v.z; //@UE
+	}
 
 	int smin, smax;
-	return dtIntersectSegmentPoly2D(spos, epos, verts, poly->vertCount, tmin, tmax, smin, smax);
+	return dtIntersectSegmentPoly2D(sposLocal, eposLocal, verts, poly->vertCount, tmin, tmax, smin, smax); //@UE
 }
 
 static void addSegmentIntersections(const dtOffMeshSegmentIntersection* isec, dtOffMeshSegmentIntersection* list, int& nlist)
@@ -476,13 +494,19 @@ static void gatherSegmentIntersections(dtMeshTile* tile,
 	for (int i = 0; i < tile->header->offMeshBase; i++)
 	{
 		dtPoly* poly = &tile->polys[i];
-		dtVcopy(bmin, &tile->verts[poly->verts[0]*3]);
-		dtVcopy(bmax, &tile->verts[poly->verts[0]*3]);
+
+		//@UE BEGIN
+		dtReal vtmp[3];
+		tile->verts[poly->verts[0]].toWorld(tile->header->bmin, vtmp);
+		dtVcopy(bmin, vtmp);
+		dtVcopy(bmax, vtmp);
 		for (int j = 1; j < poly->vertCount; j++)
 		{
-			dtVmin(bmin, &tile->verts[poly->verts[j]*3]);
-			dtVmax(bmax, &tile->verts[poly->verts[j]*3]);
+			tile->verts[poly->verts[j]].toWorld(tile->header->bmin, vtmp);
+			dtVmin(bmin, vtmp);
+			dtVmax(bmax, vtmp);
 		}
+		//@UE END
 
 		// use simple AABB overlap test first
 		if (dtOverlapBounds(segBoundsMin, segBoundsMax, bmin, bmax))
@@ -814,11 +838,22 @@ static void createSegmentPolys(dtNavMesh* nav, dtMeshTile* tile, dtOffMeshSegmen
 	{
 		dtOffMeshSegmentPart* it = &parts[i];
 
+		//@UE BEGIN
 		// add verts
-		dtVmad(&tile->verts[(vertBase+it->vA0)*3], con->startA, lenA, it->tA0);
-		dtVmad(&tile->verts[(vertBase+it->vA1)*3], con->startA, lenA, it->tA1);
-		dtVmad(&tile->verts[(vertBase+it->vB0)*3], con->startB, lenB, it->tB0);
-		dtVmad(&tile->verts[(vertBase+it->vB1)*3], con->startB, lenB, it->tB1);
+		dtReal wp[3];
+
+		dtVmad(wp, con->startA, lenA, it->tA0);
+		tile->verts[vertBase+it->vA0].setFromWorld(tile->header->bmin, wp);
+
+		dtVmad(wp, con->startA, lenA, it->tA1);
+		tile->verts[vertBase+it->vA1].setFromWorld(tile->header->bmin, wp);
+
+		dtVmad(wp, con->startB, lenB, it->tB0);
+		tile->verts[vertBase+it->vB0].setFromWorld(tile->header->bmin, wp);
+
+		dtVmad(wp, con->startB, lenB, it->tB1);
+		tile->verts[vertBase+it->vB1].setFromWorld(tile->header->bmin, wp);
+		//@UE END
 
 		// add poly
 		dtPoly* poly = &tile->polys[polyBase + i];
@@ -1072,14 +1107,16 @@ int dtNavMesh::findConnectingPolys(const dtReal* va, const dtReal* vb,
 			// Skip edges which do not point to the right side.
 			if (poly->neis[j] != m) continue;
 
-			const dtReal* vc = &tile->verts[poly->verts[j] * 3];
+			dtReal vc[3]; //@UE
+			tile->verts[poly->verts[j]].toWorld(tile->header->bmin, vc); //@UE
 			const dtReal bpos = getSlabCoord(vc, side);
 
 			// Segments are not close enough.
 			if (dtAbs(apos - bpos) > 0.01f)
 				continue;
 
-			const dtReal* vd = &tile->verts[poly->verts[(j + 1) % nv] * 3];
+			dtReal vd[3]; //@UE
+			tile->verts[poly->verts[(j + 1) % nv]].toWorld(tile->header->bmin, vd); //@UE
 
 			// Check if the segments touch.
 			calcSlabEndPoints(vc, vd, bmin, bmax, side);
@@ -1191,8 +1228,9 @@ void dtNavMesh::connectExtLinks(dtMeshTile* tile, dtMeshTile* target, int side,
 				continue;
 			
 			// Create new links
-			const dtReal* va = &tile->verts[poly->verts[j]*3];
-			const dtReal* vb = &tile->verts[poly->verts[(j+1) % nv]*3];
+			dtReal va[3], vb[3]; //@UE
+			tile->verts[poly->verts[j]].toWorld(tile->header->bmin, va); //@UE
+			tile->verts[poly->verts[(j+1) % nv]].toWorld(tile->header->bmin, vb); //@UE
 
 			// reset array before using
 			cons.resize(0);
@@ -1328,8 +1366,7 @@ void dtNavMesh::connectExtOffMeshLinks(dtMeshTile* tile, dtMeshTile* target, int
 		// For reversed links, the landing vertex is verts[0] instead of verts[1].
 		const bool bReversed = targetCon->getIsReversed();
 		const int landingVertIdx = bReversed ? 0 : 1;
-		dtReal* v = &target->verts[targetPoly->verts[landingVertIdx] * 3];
-		dtVcopy(v, nearestPt);
+		target->verts[targetPoly->verts[landingVertIdx]].setFromWorld(target->header->bmin, nearestPt); //@UE
 
 		// Link off-mesh connection to landing poly.
 		// For reversed links, this is the backtracker direction.
@@ -1471,8 +1508,7 @@ void dtNavMesh::baseOffMeshLinks(dtMeshTile* tile)
 		// Make sure the location is on current mesh.
 		const bool bReversed = con->getIsReversed();
 		const int baseVertIdx = bReversed ? 1 : 0;
-		dtReal* v = &tile->verts[poly->verts[baseVertIdx] * 3];
-		dtVcopy(v, nearestPt);
+		tile->verts[poly->verts[baseVertIdx]].setFromWorld(tile->header->bmin, nearestPt); //@UE
 
 		unsigned char sideFwd = DT_CONNECTION_INTERNAL | DT_LINK_FLAG_OFFMESH_CON | (con->getBiDirectional() ? DT_LINK_FLAG_OFFMESH_CON_BIDIR : 0);
 		unsigned char sideBck = sideFwd | DT_LINK_FLAG_OFFMESH_CON_BACKTRACKER;
@@ -1584,11 +1620,21 @@ void dtNavMesh::closestPointOnPolyInTile(const dtMeshTile* tile, unsigned int ip
 	// Off-mesh connections don't have detail polygons.
 	if (poly->getType() == DT_POLYTYPE_OFFMESH_POINT)
 	{
-		const dtReal* v0 = &tile->verts[poly->verts[0]*3];
-		const dtReal* v1 = &tile->verts[poly->verts[1]*3];
+		//@UE BEGIN
+		dtReal v0[3], v1[3];
+		tile->verts[poly->verts[0]].toWorld(tile->header->bmin, v0);
+		tile->verts[poly->verts[1]].toWorld(tile->header->bmin, v1);
 		const dtReal d0 = dtVdist(pos, v0);
 		const dtReal d1 = dtVdist(pos, v1);
-		const dtReal u = d0 / (d0+d1);
+
+		dtReal u;
+		const dtReal sum = d0 + d1;
+		if (sum != 0)
+			u = d0 / sum;
+		else
+			u = 0;
+		//@UE END
+
 		dtVlerp(closest, v0, v1, u);
 		return;
 	}
@@ -1599,7 +1645,9 @@ void dtNavMesh::closestPointOnPolyInTile(const dtMeshTile* tile, unsigned int ip
 	dtReal edget[DT_VERTS_PER_POLYGON];
 	const int nv = poly->vertCount;
 	for (int i = 0; i < nv; ++i)
-		dtVcopy(&verts[i*3], &tile->verts[poly->verts[i]*3]);
+	{ //@UE
+		tile->verts[poly->verts[i]].toWorld(tile->header->bmin, &verts[i*3]); //@UE
+	} //@UE
 	
 	dtVcopy(closest, pos);
 	if (nv == 0)
@@ -1631,21 +1679,20 @@ void dtNavMesh::closestPointOnPolyInTile(const dtMeshTile* tile, unsigned int ip
 	{
 		const dtPolyDetail* pd = &tile->detailMeshes[ip];
 
+		//@UE BEGIN
 		for (int j = 0; j < pd->triCount; ++j)
 		{
 			const unsigned char* t = &tile->detailTris[(pd->triBase+j)*4];
+			dtReal vBuf[3][3];
 			const dtReal* v[3];
 			for (int k = 0; k < 3; ++k)
 			{
+				CA_SUPPRESS(6385);
 				if (t[k] < poly->vertCount)
-				{
-					CA_SUPPRESS(6385);
-					v[k] = &tile->verts[poly->verts[t[k]]*3];
-				}
+					tile->verts[poly->verts[t[k]]].toWorld(tile->header->bmin, vBuf[k]);
 				else
-				{
-					v[k] = &tile->detailVerts[(pd->vertBase+(t[k]-poly->vertCount))*3];
-				}
+					tile->detailVerts[pd->vertBase+(t[k]-poly->vertCount)].toWorld(tile->header->bmin, vBuf[k]);
+				v[k] = vBuf[k];
 			}
 			dtReal h;
 			if (dtClosestHeightPointTriangle(closest, v[0], v[1], v[2], h))
@@ -1654,6 +1701,7 @@ void dtNavMesh::closestPointOnPolyInTile(const dtMeshTile* tile, unsigned int ip
 				break;
 			}
 		}
+		//@UE END
 	}
 	else
 	{
@@ -1841,12 +1889,13 @@ int dtNavMesh::queryPolygonsInTile(const dtMeshTile* tile, const dtReal* qmin, c
 				continue;
 
 			// Calc polygon bounds.
-			const dtReal* v = &tile->verts[p->verts[0]*3];
+			dtReal v[3]; //@UE
+			tile->verts[p->verts[0]].toWorld(tile->header->bmin, v); //@UE
 			dtVcopy(bmin, v);
 			dtVcopy(bmax, v);
 			for (int j = 1; j < p->vertCount; ++j)
 			{
-				v = &tile->verts[p->verts[j]*3];
+				tile->verts[p->verts[j]].toWorld(tile->header->bmin, v); //@UE
 				dtVmin(bmin, v);
 				dtVmax(bmax, v);
 			}
@@ -1934,11 +1983,11 @@ dtStatus dtNavMesh::addTile(unsigned char* data, int dataSize, int flags,
 	
 	// Patch header pointers.
 	const int headerSize = dtAlign(sizeof(dtMeshHeader));
-	const int vertsSize = dtAlign(sizeof(dtReal)*3*header->vertCount);
+	const int vertsSize = dtAlign(sizeof(dtTileVert)*header->vertCount); //@UE
 	const int polysSize = dtAlign(sizeof(dtPoly)*header->polyCount);
 	const int linksSize = dtAlign(sizeof(dtLink)*(header->maxLinkCount));
 	const int detailMeshesSize = dtAlign(sizeof(dtPolyDetail)*header->detailMeshCount);
-	const int detailVertsSize = dtAlign(sizeof(dtReal)*3*header->detailVertCount);
+	const int detailVertsSize = dtAlign(sizeof(dtTileVert)*header->detailVertCount); //@UE
 	const int detailTrisSize = dtAlign(sizeof(unsigned char)*4*header->detailTriCount);
 	const int bvtreeSize = dtAlign(sizeof(dtBVNode)*header->bvNodeCount);
 	const int offMeshLinksSize = dtAlign(sizeof(dtOffMeshConnection)*header->offMeshConCount);
@@ -1953,11 +2002,11 @@ dtStatus dtNavMesh::addTile(unsigned char* data, int dataSize, int flags,
 #endif // WITH_NAVMESH_CLUSTER_LINKS
 
 	const unsigned char* d = data + headerSize;
-	tile->verts = (dtReal*)d; d += vertsSize;
+	tile->verts = (dtTileVert*)d; d += vertsSize; //@UE
 	tile->polys = (dtPoly*)d; d += polysSize;
 	tile->links = (dtLink*)d; d += linksSize;
 	tile->detailMeshes = (dtPolyDetail*)d; d += detailMeshesSize;
-	tile->detailVerts = (dtReal*)d; d += detailVertsSize;
+	tile->detailVerts = (dtTileVert*)d; d += detailVertsSize; //@UE
 	tile->detailTris = (unsigned char*)d; d += detailTrisSize;
 	tile->bvTree = (dtBVNode*)d; d += bvtreeSize;
 	tile->offMeshCons = (dtOffMeshConnection*)d; d += offMeshLinksSize;
@@ -2812,8 +2861,8 @@ dtStatus dtNavMesh::getOffMeshConnectionPolyEndPoints(dtPolyRef prevRef, dtPolyR
 		const int idx3 = (idx1 == 1) ? 3 : 1;
 
 		dtReal start0[3], start1[3];
-		dtVcopy(start0, &tile->verts[poly->verts[idx0]*3]);
-		dtVcopy(start1, &tile->verts[poly->verts[idx1]*3]);
+		tile->verts[poly->verts[idx0]].toWorld(tile->header->bmin, start0); //@UE
+		tile->verts[poly->verts[idx1]].toWorld(tile->header->bmin, start1); //@UE
 		dtReal t = 0;
 		
 		// Find the point on the segment link edge that is closest to the agent's current position.
@@ -2832,8 +2881,8 @@ dtStatus dtNavMesh::getOffMeshConnectionPolyEndPoints(dtPolyRef prevRef, dtPolyR
 		// We want to use the shortest path to cross the segment link as this is generally the most sensible thing to do.
 		// This change also matches the change to the string pulling algorithm used by traditional path following.
 		dtReal end0[3], end1[3];
-		dtVcopy(end0, &tile->verts[poly->verts[idx2] * 3]);
-		dtVcopy(end1, &tile->verts[poly->verts[idx3] * 3]);
+		tile->verts[poly->verts[idx2]].toWorld(tile->header->bmin, end0); //@UE
+		tile->verts[poly->verts[idx3]].toWorld(tile->header->bmin, end1); //@UE
 		dtDistancePtSegSqr2D(startPos, end0, end1, t);
 		dtVlerp(endPos, end0, end1, t);
 	}
@@ -2841,8 +2890,8 @@ dtStatus dtNavMesh::getOffMeshConnectionPolyEndPoints(dtPolyRef prevRef, dtPolyR
 #endif // WITH_NAVMESH_SEGMENT_LINKS
 	//@UE END
 	{
-		dtVcopy(startPos, &tile->verts[poly->verts[idx0]*3]);
-		dtVcopy(endPos, &tile->verts[poly->verts[idx1]*3]);
+		tile->verts[poly->verts[idx0]].toWorld(tile->header->bmin, startPos); //@UE
+		tile->verts[poly->verts[idx1]].toWorld(tile->header->bmin, endPos); //@UE
 	}
 
 	return DT_SUCCESS;
@@ -3027,18 +3076,8 @@ void dtNavMesh::applyWorldOffset(const dtReal* offset)
 			// Shift tile bounds
 			dtVadd(tile.header->bmin, tile.header->bmin, offset);
 			dtVadd(tile.header->bmax, tile.header->bmax, offset);
-			
-			//Shift tile vertices
-			for (int j = 0; j < tile.header->vertCount; ++j)
-			{
-				dtVadd(&(tile.verts[j*3]), &(tile.verts[j*3]), offset);
-			}
-			
-			//Shift tile details vertices
-			for (int j = 0; j < tile.header->detailVertCount; ++j)
-			{
-				dtVadd(&(tile.detailVerts[j*3]), &(tile.detailVerts[j*3]), offset);
-			}
+
+			// Tile vertices are stored as local offsets from bmin; shifting bmin above is sufficient.
 
 			//Shift off-mesh connections
 			for (int j = 0; j < tile.header->offMeshConCount; ++j)
diff --git a/Engine/Source/Runtime/Navmesh/Private/Detour/DetourNavMeshBuilder.cpp b/Engine/Source/Runtime/Navmesh/Private/Detour/DetourNavMeshBuilder.cpp
--- a/Engine/Source/Runtime/Navmesh/Private/Detour/DetourNavMeshBuilder.cpp
+++ b/Engine/Source/Runtime/Navmesh/Private/Detour/DetourNavMeshBuilder.cpp
@@ -168,12 +168,32 @@ static void subdivide(BVItem* items, int nitems, int imin, int imax, int& curNod
 
 static int createBVTree(const unsigned short* verts, const int /*nverts*/,
 						const unsigned short* polys, const int npolys, const int nvp,
-						const dtPolyDetail* DMeshes, const dtReal* DVerts, const unsigned char* DTris, const dtReal* tbmin,
+						const dtPolyDetail* DMeshes, const dtTileVert* DVerts, const unsigned char* DTris, //@UE
 						const dtReal cs, const dtReal ch,
 						const int /*nnodes*/, dtBVNode* nodes)
 {
+	//@UE BEGIN
+	const uint64 itemSize = sizeof(BVItem)*npolys;
+	if (itemSize > INT_MAX)
+	{
+		return 0;
+	}
+
 	// Build tree
-	BVItem* items = (BVItem*)dtAlloc(sizeof(BVItem)*npolys, DT_ALLOC_TEMP);
+	BVItem* items = (BVItem*)dtAlloc(itemSize, DT_ALLOC_TEMP);
+
+	if (items == nullptr)
+	{
+		return 0;
+	}
+
+	if (cs <= 0 || ch <= 0)
+	{
+		dtFree(items, DT_ALLOC_TEMP);
+		return 0;
+	}
+	//@UE END
+
 	for (int i = 0; i < npolys; i++)
 	{
 		BVItem& it = items[i];
@@ -215,8 +235,7 @@ static int createBVTree(const unsigned short* verts, const int /*nverts*/,
 			{
 				if (t[m] >= vertCount)
 				{
-					const dtReal* detailCoords = &DVerts[(pd->vertBase + (t[m] - vertCount)) * 3];
-					const dtReal qY = (detailCoords[1] - tbmin[1]) / ch;
+					const dtReal qY = DVerts[pd->vertBase + (t[m] - vertCount)].y / ch; //@UE
 					const unsigned short qYmin = (unsigned short)dtFloor(qY);
 					const unsigned short qYmax = (unsigned short)dtCeil(qY);
 
@@ -463,11 +482,11 @@ bool dtCreateNavMeshData(dtNavMeshCreateParams* params, unsigned char** outData,
  
 	// Calculate data size
 	const int headerSize = dtAlign(sizeof(dtMeshHeader));
-	const int vertsSize = dtAlign(sizeof(dtReal)*3*totVertCount);
+	const int vertsSize = dtAlign(sizeof(dtTileVert)*totVertCount); //@UE
 	const int polysSize = dtAlign(sizeof(dtPoly)*totPolyCount);
 	const int linksSize = dtAlign(sizeof(dtLink)*maxLinkCount);
 	const int detailMeshesSize = dtAlign(sizeof(dtPolyDetail)*params->polyCount);
-	const int detailVertsSize = dtAlign(sizeof(dtReal)*3*uniqueDetailVertCount);
+	const int detailVertsSize = dtAlign(sizeof(dtTileVert)*uniqueDetailVertCount); //@UE
 	const int detailTrisSize = dtAlign(sizeof(unsigned char)*4*detailTriCount);
 	const int bvTreeSize = params->buildBvTree ? dtAlign(sizeof(dtBVNode)*params->polyCount*2) : 0;
 	const int offMeshConsSize = dtAlign(sizeof(dtOffMeshConnection)*storedOffMeshConCount);
@@ -503,11 +522,11 @@ bool dtCreateNavMeshData(dtNavMeshCreateParams* params, unsigned char** outData,
 	
 	unsigned char* d = data;
 	dtMeshHeader* header = (dtMeshHeader*)d; d += headerSize;
-	dtReal* navVerts = (dtReal*)d; d += vertsSize;
+	dtTileVert* navVerts = (dtTileVert*)d; d += vertsSize; //@UE
 	dtPoly* navPolys = (dtPoly*)d; d += polysSize;
 	d += linksSize;
 	dtPolyDetail* navDMeshes = (dtPolyDetail*)d; d += detailMeshesSize;
-	dtReal* navDVerts = (dtReal*)d; d += detailVertsSize;
+	dtTileVert* navDVerts = (dtTileVert*)d; d += detailVertsSize; //@UE
 	unsigned char* navDTris = (unsigned char*)d; d += detailTrisSize;
 	dtBVNode* navBvtree = (dtBVNode*)d; d += bvTreeSize;
 	dtOffMeshConnection* offMeshCons = (dtOffMeshConnection*)d; d += offMeshConsSize;
@@ -589,14 +608,13 @@ bool dtCreateNavMeshData(dtNavMeshCreateParams* params, unsigned char** outData,
 	const int offMeshPolyBase = params->polyCount;
 
 	// Store vertices
-	// Mesh vertices
+	// Mesh vertices: stored as tile-local float offsets from params->bmin //@UE
 	for (int i = 0; i < params->vertCount; ++i)
 	{
 		const unsigned short* iv = &params->verts[i*3];
-		dtReal* v = &navVerts[i*3];
-		v[0] = params->bmin[0] + iv[0] * params->cs;
-		v[1] = params->bmin[1] + iv[1] * params->ch;
-		v[2] = params->bmin[2] + iv[2] * params->cs;
+		navVerts[i].x = (float)(iv[0] * params->cs); //@UE
+		navVerts[i].y = (float)(iv[1] * params->ch); //@UE
+		navVerts[i].z = (float)(iv[2] * params->cs); //@UE
 	}
 	// Off-mesh point link vertices.
 	int n = 0;
@@ -607,9 +625,8 @@ bool dtCreateNavMeshData(dtNavMeshCreateParams* params, unsigned char** outData,
 		// Only store connections which start from this tile.
 		if ((offMeshConClass[i*2+0] == 0xff) && (offMeshCon.type & DT_OFFMESH_CON_POINT))
 		{
-			dtReal* v = &navVerts[(offMeshVertsBase + n*2)*3];
-			dtVcopy(&v[0], &offMeshCon.vertsA0[0]);
-			dtVcopy(&v[3], &offMeshCon.vertsB0[0]);
+			navVerts[offMeshVertsBase + n*2    ].setFromWorld(params->bmin, &offMeshCon.vertsA0[0]); //@UE
+			navVerts[offMeshVertsBase + n*2 + 1].setFromWorld(params->bmin, &offMeshCon.vertsB0[0]); //@UE
 			n++;
 		}
 	}
@@ -712,7 +729,10 @@ bool dtCreateNavMeshData(dtNavMeshCreateParams* params, unsigned char** outData,
 			// Copy vertices except the first 'nv' verts which are equal to nav poly verts.
 			if (ndv-nv)
 			{
-				memcpy(&navDVerts[vbase*3], &params->detailVerts[(vb+nv)*3], sizeof(dtReal)*3*(ndv-nv));
+				for (int k = 0; k < (ndv-nv); ++k) //@UE
+				{ //@UE
+					navDVerts[vbase+k].setFromWorld(params->bmin, &params->detailVerts[(vb+nv+k)*3]); //@UE
+				} //@UE
 				vbase += (unsigned short)(ndv-nv);
 			}
 		}
@@ -750,9 +770,17 @@ bool dtCreateNavMeshData(dtNavMeshCreateParams* params, unsigned char** outData,
 	// Store and create BVtree.
 	if (params->buildBvTree)
 	{
-		createBVTree(params->verts, params->vertCount, params->polys, params->polyCount, nvp,
-					 navDMeshes, navDVerts, navDTris, params->bmin,
+		//@UE BEGIN
+		const int nodeCount = createBVTree(params->verts, params->vertCount, params->polys, params->polyCount, nvp,
+					 navDMeshes, navDVerts, navDTris,
 					 params->cs, params->ch, params->polyCount*2, navBvtree);
+		if (nodeCount == 0)
+		{
+			dtFree(data, DT_ALLOC_PERM_TILE_DATA);
+			dtFree(offMeshConClass, DT_ALLOC_TEMP);
+			return false;
+		}
+		//@UE END
 	}
 	
 	// Store Off-Mesh connections.
@@ -830,12 +858,23 @@ bool dtCreateNavMeshData(dtNavMeshCreateParams* params, unsigned char** outData,
 			const dtPoly* poly = &navPolys[j];
 			dtReal c[3] = { 0.0f, 0.0f, 0.0f };
 
+			// Sum vertices in local space
 			for (int iv = 0; iv < poly->vertCount; iv++)
 			{
-				dtVadd(c, c, &navVerts[poly->verts[iv] * 3]);
+				const dtTileVert& v = navVerts[poly->verts[iv]];
+				c[0] += (dtReal)v.x;
+				c[1] += (dtReal)v.y;
+				c[2] += (dtReal)v.z;
 			}
-			
-			dtVmad(cluster.center, cluster.center, c, dtReal(1.) / poly->vertCount);
+
+			// Convert centroid to world space.
+			const dtReal invCount = dtReal(1.) / poly->vertCount;
+			dtReal cWorld[3];
+			cWorld[0] = c[0] * invCount + params->bmin[0];
+			cWorld[1] = c[1] * invCount + params->bmin[1];
+			cWorld[2] = c[2] * invCount + params->bmin[2];
+
+			dtVmad(cluster.center, cluster.center, cWorld, dtReal(1.));
 			break;
 		}
 	}
@@ -913,11 +952,11 @@ bool dtNavMeshDataSwapEndian(unsigned char* data, const int /*dataSize*/)
 	
 	// Patch header pointers.
 	const int headerSize = dtAlign(sizeof(dtMeshHeader));
-	const int vertsSize = dtAlign(sizeof(dtReal)*3*header->vertCount);
+	const int vertsSize = dtAlign(sizeof(dtTileVert)*header->vertCount); //@UE
 	const int polysSize = dtAlign(sizeof(dtPoly)*header->polyCount);
 	const int linksSize = dtAlign(sizeof(dtLink)*(header->maxLinkCount));
 	const int detailMeshesSize = dtAlign(sizeof(dtPolyDetail)*header->detailMeshCount);
-	const int detailVertsSize = dtAlign(sizeof(dtReal)*3*header->detailVertCount);
+	const int detailVertsSize = dtAlign(sizeof(dtTileVert)*header->detailVertCount); //@UE
 	const int detailTrisSize = dtAlign(sizeof(unsigned char)*4*header->detailTriCount);
 	const int bvtreeSize = dtAlign(sizeof(dtBVNode)*header->bvNodeCount);
 	const int offMeshLinksSize = dtAlign(sizeof(dtOffMeshConnection)*header->offMeshConCount);
@@ -934,11 +973,11 @@ bool dtNavMeshDataSwapEndian(unsigned char* data, const int /*dataSize*/)
 	//@UE END
 
 	unsigned char* d = data + headerSize;
-	dtReal* verts = (dtReal*)d; d += vertsSize;
+	dtTileVert* verts = (dtTileVert*)d; d += vertsSize; //@UE
 	dtPoly* polys = (dtPoly*)d; d += polysSize;
 	/*dtLink* links = (dtLink*)d;*/ d += linksSize;
 	dtPolyDetail* detailMeshes = (dtPolyDetail*)d; d += detailMeshesSize;
-	dtReal* detailVerts = (dtReal*)d; d += detailVertsSize;
+	dtTileVert* detailVerts = (dtTileVert*)d; d += detailVertsSize; //@UE
 	/*unsigned char* detailTris = (unsigned char*)d;*/ d += detailTrisSize;
 	dtBVNode* bvTree = (dtBVNode*)d; d += bvtreeSize;
 	dtOffMeshConnection* offMeshCons = (dtOffMeshConnection*)d; d += offMeshLinksSize;
@@ -955,9 +994,11 @@ bool dtNavMeshDataSwapEndian(unsigned char* data, const int /*dataSize*/)
 	//@UE END
 
 	// Vertices
-	for (int i = 0; i < header->vertCount*3; ++i)
+	for (int i = 0; i < header->vertCount; ++i)
 	{
-		dtSwapEndian(&verts[i]);
+		dtSwapEndian(&verts[i].x); //@UE
+		dtSwapEndian(&verts[i].y); //@UE
+		dtSwapEndian(&verts[i].z); //@UE
 	}
 
 	// Polys
@@ -984,9 +1025,11 @@ bool dtNavMeshDataSwapEndian(unsigned char* data, const int /*dataSize*/)
 	}
 	
 	// Detail verts
-	for (int i = 0; i < header->detailVertCount*3; ++i)
+	for (int i = 0; i < header->detailVertCount; ++i)
 	{
-		dtSwapEndian(&detailVerts[i]);
+		dtSwapEndian(&detailVerts[i].x); //@UE
+		dtSwapEndian(&detailVerts[i].y); //@UE
+		dtSwapEndian(&detailVerts[i].z); //@UE
 	}
 
 	// BV-tree
@@ -1046,6 +1089,10 @@ bool dtNavMeshDataSwapEndian(unsigned char* data, const int /*dataSize*/)
 // Experimental tile transform
 bool dtTransformTileData(unsigned char* data, const int dataSize, const int offsetX, const int offsetY, const dtReal tileWidth, const dtReal tileHeight, const dtReal rotationDeg, const dtReal bvQuantFactor)
 {
+	// Validate minimum size before casting to header
+	if (dataSize < (int)sizeof(dtMeshHeader))
+		return false;
+
 	// Make sure the data is in right format.
 	dtMeshHeader* header = (dtMeshHeader*)data;
 	if (header->version != DT_NAVMESH_VERSION)
@@ -1057,11 +1104,11 @@ bool dtTransformTileData(unsigned char* data, const int dataSize, const int offs
 
 	// Patch header pointers.
 	const int headerSize = dtAlign(sizeof(dtMeshHeader));
-	const int vertsSize = dtAlign(sizeof(dtReal)*3*header->vertCount);
+	const int vertsSize = dtAlign(sizeof(dtTileVert)*header->vertCount);
 	const int polysSize = dtAlign(sizeof(dtPoly)*header->polyCount);
 	const int linksSize = dtAlign(sizeof(dtLink)*(header->maxLinkCount));
 	const int detailMeshesSize = dtAlign(sizeof(dtPolyDetail)*header->detailMeshCount);
-	const int detailVertsSize = dtAlign(sizeof(dtReal)*3*header->detailVertCount);
+	const int detailVertsSize = dtAlign(sizeof(dtTileVert)*header->detailVertCount);
 	const int detailTrisSize = dtAlign(sizeof(unsigned char)*4*header->detailTriCount);
 	const int bvtreeSize = dtAlign(sizeof(dtBVNode)*header->bvNodeCount);
 	const int offMeshLinksSize = dtAlign(sizeof(dtOffMeshConnection)*header->offMeshConCount);
@@ -1079,11 +1126,11 @@ bool dtTransformTileData(unsigned char* data, const int dataSize, const int offs
 	tile.header = header;
 
 	unsigned char* d = data + headerSize;
-	tile.verts = (dtReal*)d; d += vertsSize;
+	tile.verts = (dtTileVert*)d; d += vertsSize;
 	tile.polys = (dtPoly*)d; d += polysSize;
 	tile.links = (dtLink*)d; d += linksSize;
 	tile.detailMeshes = (dtPolyDetail*)d; d += detailMeshesSize;
-	tile.detailVerts = (dtReal*)d; d += detailVertsSize;
+	tile.detailVerts = (dtTileVert*)d; d += detailVertsSize;
 	tile.detailTris = (unsigned char*)d; d += detailTrisSize;
 	tile.bvTree = (dtBVNode*)d; d += bvtreeSize;
 	tile.offMeshCons = (dtOffMeshConnection*)d; d += offMeshLinksSize;
@@ -1097,6 +1144,17 @@ bool dtTransformTileData(unsigned char* data, const int dataSize, const int offs
 	tile.polyClusters = (unsigned short*)d; d += clusterPolysSize;
 #endif //WITH_NAVMESH_CLUSTER_LINKS
 
+	// Validate that the computed layout fits within the provided buffer
+	const int computedSize = (int)(d - data);
+	if (computedSize > dataSize)
+	{
+		return false;
+	}
+
+	// Save old bmin before transforming the header; needed to toWorld tile-local vertex offsets
+	dtReal oldBmin[3];
+	dtVcopy(oldBmin, header->bmin);
+
 	// Rotate on original center position
 	dtReal rotationCenter[3];
 	dtVadd(rotationCenter, header->bmin, header->bmax);
@@ -1157,18 +1215,24 @@ bool dtTransformTileData(unsigned char* data, const int dataSize, const int offs
 		}
 	}
 
-	// Transform tile vertices
+	// Transform tile vertices (transform from old bmin, rotate, transform into new bmin)
 	for (int j = 0; j < tile.header->vertCount; ++j)
 	{
-		dtRotate90(&(tile.verts[j*3]), &(tile.verts[j*3]), rotationCenter, rot);
-		dtVadd(&(tile.verts[j*3]), &(tile.verts[j*3]), offset);
+		dtReal v[3], vOut[3];
+		tile.verts[j].toWorld(oldBmin, v); //@UE
+		dtRotate90(vOut, v, rotationCenter, rot);
+		dtVadd(vOut, vOut, offset);
+		tile.verts[j].setFromWorld(tile.header->bmin, vOut); //@UE
 	}
 
-	// Transform tile details vertices
+	// Transform tile detail vertices (transform from old bmin, rotate, transform into new bmin)
 	for (int j = 0; j < tile.header->detailVertCount; ++j)
 	{
-		dtRotate90(&(tile.detailVerts[j*3]), &(tile.detailVerts[j*3]), rotationCenter, rot);
-		dtVadd(&(tile.detailVerts[j*3]), &(tile.detailVerts[j*3]), offset);
+		dtReal v[3], vOut[3];
+		tile.detailVerts[j].toWorld(oldBmin, v); //@UE
+		dtRotate90(vOut, v, rotationCenter, rot);
+		dtVadd(vOut, vOut, offset);
+		tile.detailVerts[j].setFromWorld(tile.header->bmin, vOut); //@UE
 	}
 
 	// Transfrom BVTree (bmin and bmax are local to the tile)
diff --git a/Engine/Source/Runtime/Navmesh/Private/Detour/DetourNavMeshQuery.cpp b/Engine/Source/Runtime/Navmesh/Private/Detour/DetourNavMeshQuery.cpp
--- a/Engine/Source/Runtime/Navmesh/Private/Detour/DetourNavMeshQuery.cpp
+++ b/Engine/Source/Runtime/Navmesh/Private/Detour/DetourNavMeshQuery.cpp
@@ -106,6 +106,16 @@ void dtQueryFilterData::copyFrom(const dtQueryFilterData* source)
 
 //@UE BEGIN
 static const dtReal DEFAULT_HEURISTIC_SCALE = 0.999f; // Search heuristic scale.
+
+/// Compute triangle area in tile-local space.
+inline dtReal dtTriArea2DLocal(const dtTileVert& a, const dtTileVert& b, const dtTileVert& c)
+{
+	const float abx = b.x - a.x;
+	const float abz = b.z - a.z;
+	const float acx = c.x - a.x;
+	const float acz = c.z - a.z;
+	return (dtReal)(acx * abz - abx * acz);
+}
 //@UE END
 
 dtNavMeshQuery* dtAllocNavMeshQuery()
@@ -355,10 +365,11 @@ dtStatus dtNavMeshQuery::findRandomPoint(const dtQueryFilter* filter, float(*fra
 		dtReal polyArea = 0.0f;
 		for (int j = 2; j < p->vertCount; ++j)
 		{
-			const dtReal* va = &tile->verts[p->verts[0]*3];
-			const dtReal* vb = &tile->verts[p->verts[j-1]*3];
-			const dtReal* vc = &tile->verts[p->verts[j]*3];
-			polyArea += dtTriArea2D(va,vb,vc);
+			polyArea += dtTriArea2DLocal( //@UE
+				tile->verts[p->verts[0]], //@UE
+				tile->verts[p->verts[j-1]], //@UE
+				tile->verts[p->verts[j]] //@UE
+			); //@UE
 		}
 
 		// Choose random polygon weighted by area, using reservoi sampling.
@@ -379,8 +390,7 @@ dtStatus dtNavMeshQuery::findRandomPoint(const dtQueryFilter* filter, float(*fra
 	dtReal areas[DT_VERTS_PER_POLYGON];
 	for (int j = 0; j < poly->vertCount; ++j)
 	{
-		const dtReal* v = &tile->verts[poly->verts[j]*3];
-		dtVcopy(&verts[j*3],v);
+		tile->verts[poly->verts[j]].toWorld(tile->header->bmin, &verts[j*3]); //@UE
 	}
 	
 	const dtReal s = frand();
@@ -462,10 +472,11 @@ dtStatus dtNavMeshQuery::findRandomPointAroundCircle(dtPolyRef startRef, const d
 			dtReal polyArea = 0.0f;
 			for (int j = 2; j < bestPoly->vertCount; ++j)
 			{
-				const dtReal* va = &bestTile->verts[bestPoly->verts[0]*3];
-				const dtReal* vb = &bestTile->verts[bestPoly->verts[j-1]*3];
-				const dtReal* vc = &bestTile->verts[bestPoly->verts[j]*3];
-				polyArea += dtTriArea2D(va,vb,vc);
+				polyArea += dtTriArea2DLocal( //@UE
+					bestTile->verts[bestPoly->verts[0]], //@UE
+					bestTile->verts[bestPoly->verts[j-1]], //@UE
+					bestTile->verts[bestPoly->verts[j]] //@UE
+				); //@UE
 			}
 			// Choose random polygon weighted by area, using reservoi sampling.
 			areaSum += polyArea;
@@ -573,8 +584,7 @@ dtStatus dtNavMeshQuery::findRandomPointAroundCircle(dtPolyRef startRef, const d
 
 		for (int j = 0; j < testPoly->vertCount; ++j)
 		{
-			const dtReal* v = &testTile->verts[testPoly->verts[j]*3];
-			dtVcopy(&verts[j*3],v);
+			testTile->verts[testPoly->verts[j]].toWorld(testTile->header->bmin, &verts[j*3]); //@UE
 		}
 
 		dtReal pt[3];
@@ -615,8 +625,9 @@ dtStatus dtNavMeshQuery::findRandomPointInPoly(dtPolyRef ref, float(*frand)(), d
 
 	if (poly->getType() == DT_POLYTYPE_OFFMESH_POINT)
 	{
-		const dtReal* v0 = &tile->verts[poly->verts[0]*3];
-		const dtReal* v1 = &tile->verts[poly->verts[1]*3];
+		dtReal v0[3], v1[3]; //@UE
+		tile->verts[poly->verts[0]].toWorld(tile->header->bmin, v0); //@UE
+		tile->verts[poly->verts[1]].toWorld(tile->header->bmin, v1); //@UE
 		const dtReal s = frand();
 		dtVlerp(randomPt, v0, v1, s);
 		return DT_SUCCESS;
@@ -627,8 +638,7 @@ dtStatus dtNavMeshQuery::findRandomPointInPoly(dtPolyRef ref, float(*frand)(), d
 		dtReal areas[DT_VERTS_PER_POLYGON] = {0};
 		for (int j = 0; j < poly->vertCount; ++j)
 		{
-			const dtReal* v = &tile->verts[poly->verts[j]*3];
-			dtVcopy(&verts[j*3],v);
+			tile->verts[poly->verts[j]].toWorld(tile->header->bmin, &verts[j*3]); //@UE
 		}
 
 		const dtReal s = frand();
@@ -685,10 +695,11 @@ dtStatus dtNavMeshQuery::findRandomPointInCluster(dtClusterRef clusterRef, float
 			dtReal polyArea = 0.0f;
 			for (int j = 2; j < testPoly->vertCount; ++j)
 			{
-				const dtReal* va = &searchTile->verts[testPoly->verts[0]*3];
-				const dtReal* vb = &searchTile->verts[testPoly->verts[j-1]*3];
-				const dtReal* vc = &searchTile->verts[testPoly->verts[j]*3];
-				polyArea += dtTriArea2D(va,vb,vc);
+				polyArea += dtTriArea2DLocal( //@UE
+					searchTile->verts[testPoly->verts[0]], //@UE
+					searchTile->verts[testPoly->verts[j-1]], //@UE
+					searchTile->verts[testPoly->verts[j]] //@UE
+				); //@UE
 			}
 			// Choose random polygon weighted by area, using reservoi sampling.
 			areaSum += polyArea;
@@ -711,8 +722,7 @@ dtStatus dtNavMeshQuery::findRandomPointInCluster(dtClusterRef clusterRef, float
 	dtReal areas[DT_VERTS_PER_POLYGON];
 	for (int j = 0; j < randomPoly->vertCount; ++j)
 	{
-		const dtReal* v = &searchTile->verts[randomPoly->verts[j]*3];
-		dtVcopy(&verts[j*3],v);
+		searchTile->verts[randomPoly->verts[j]].toWorld(searchTile->header->bmin, &verts[j*3]); //@UE
 	}
 
 	const dtReal s = frand();
@@ -767,11 +777,21 @@ void dtNavMeshQuery::closestPointOnPolyInTile(const dtMeshTile* tile, const dtPo
 	// Off-mesh connections don't have detail polygons.
 	if (poly->getType() == DT_POLYTYPE_OFFMESH_POINT)
 	{
-		const dtReal* v0 = &tile->verts[poly->verts[0]*3];
-		const dtReal* v1 = &tile->verts[poly->verts[1]*3];
+		dtReal v0[3], v1[3]; //@UE
+		tile->verts[poly->verts[0]].toWorld(tile->header->bmin, v0); //@UE
+		tile->verts[poly->verts[1]].toWorld(tile->header->bmin, v1); //@UE
 		const dtReal d0 = dtVdist(pos, v0);
 		const dtReal d1 = dtVdist(pos, v1);
-		const dtReal u = d0 / (d0+d1);
+
+		//@UE BEGIN
+		dtReal u;
+		const dtReal sum = d0 + d1;
+		if (sum != 0)
+		    u = d0 / sum;
+		else
+		    u = 0;
+		//@UE END
+
 		dtVlerp(closest, v0, v1, u);
 		return;
 	}
@@ -784,9 +804,7 @@ void dtNavMeshQuery::closestPointOnPolyInTile(const dtMeshTile* tile, const dtPo
 	dtReal edget[DT_VERTS_PER_POLYGON];
 	const int nv = dtMin(poly->vertCount, DT_VERTS_PER_POLYGON);
 	for (int i = 0; i < nv; ++i)
-	{
-		dtVcopy(&verts[i*3], &tile->verts[poly->verts[i]*3]);
-	}
+		tile->verts[poly->verts[i]].toWorld(tile->header->bmin, &verts[i*3]); //@UE
 	
 	dtVcopy(closest, pos);
 	if (nv == 0)
@@ -821,18 +839,20 @@ void dtNavMeshQuery::closestPointOnPolyInTile(const dtMeshTile* tile, const dtPo
 		for (int j = 0; j < pd->triCount; ++j)
 		{
 			const unsigned char* t = &tile->detailTris[(pd->triBase+j)*4];
+			dtReal vBuf[3][3]; //@UE
 			const dtReal* v[3];
 			for (int k = 0; k < 3; ++k)
 			{
 				if (t[k] < poly->vertCount)
 				{
 					CA_SUPPRESS(6385);
-					v[k] = &tile->verts[poly->verts[t[k]]*3];
+					tile->verts[poly->verts[t[k]]].toWorld(tile->header->bmin, vBuf[k]); //@UE
 				}
 				else
 				{
-					v[k] = &tile->detailVerts[(pd->vertBase+(t[k]-poly->vertCount))*3];
+					tile->detailVerts[pd->vertBase+(t[k]-poly->vertCount)].toWorld(tile->header->bmin, vBuf[k]); //@UE
 				}
+				v[k] = vBuf[k]; //@UE
 			}
 			dtReal h;
 			if (dtClosestHeightPointTriangle(closest, v[0], v[1], v[2], h))
@@ -890,7 +910,7 @@ dtStatus dtNavMeshQuery::closestPointOnPolyBoundary(dtPolyRef ref, const dtReal*
 	int nv = 0;
 	for (int i = 0; i < (int)poly->vertCount; ++i)
 	{
-		dtVcopy(&verts[nv*3], &tile->verts[poly->verts[i]*3]);
+		tile->verts[poly->verts[i]].toWorld(tile->header->bmin, &verts[nv*3]); //@UE
 		nv++;
 	}		
 	
@@ -946,21 +966,22 @@ dtStatus dtNavMeshQuery::isPointInsidePoly(dtPolyRef ref, const dtReal* pos, boo
 	dtAssert(m_nav);
 	const dtMeshTile* tile = 0;
 	const dtPoly* poly = 0;
+	result = false; //@UE
+	
 	if (dtStatusFailed(m_nav->getTileAndPolyByRef(ref, &tile, &poly)))
 		return DT_FAILURE | DT_INVALID_PARAM;
+
 	if (!tile)
 		return DT_FAILURE | DT_INVALID_PARAM;
 
 	if (poly->getType() == DT_POLYTYPE_OFFMESH_POINT)
-		return false;
-
-	const unsigned int ip = (unsigned int)(poly - tile->polys);
+		return DT_FAILURE; //@UE
 
 	// Clamp point to be inside the polygon.
 	dtReal verts[DT_VERTS_PER_POLYGON * 3];
 	const int nv = poly->vertCount;
 	for (int i = 0; i < nv; ++i)
-		dtVcopy(&verts[i * 3], &tile->verts[poly->verts[i] * 3]);
+		tile->verts[poly->verts[i]].toWorld(tile->header->bmin, &verts[i * 3]); //@UE
 
 	result = dtPointInPolygon(pos, verts, nv);
 
@@ -973,11 +994,21 @@ dtStatus dtNavMeshQuery::projectedPointOnPolyInTile(const dtMeshTile* tile, cons
 	// Off-mesh connections don't have detail polygons.
 	if (poly->getType() == DT_POLYTYPE_OFFMESH_POINT)
 	{
-		const dtReal* v0 = &tile->verts[poly->verts[0] * 3];
-		const dtReal* v1 = &tile->verts[poly->verts[1] * 3];
+		dtReal v0[3], v1[3]; //@UE
+		tile->verts[poly->verts[0]].toWorld(tile->header->bmin, v0); //@UE
+		tile->verts[poly->verts[1]].toWorld(tile->header->bmin, v1); //@UE
 		const dtReal d0 = dtVdist(pos, v0);
 		const dtReal d1 = dtVdist(pos, v1);
-		const dtReal u = d0 / (d0 + d1);
+
+		//@UE BEGIN
+		dtReal u;
+		const dtReal sum = d0 + d1;
+		if (sum != 0)
+			u = d0 / sum;
+		else
+			u = 0;
+		//@UE END
+
 		dtVlerp(projected, v0, v1, u);
 		
 		// @todo this is not quite true, this calculated the closes point, not a projection
@@ -990,7 +1021,7 @@ dtStatus dtNavMeshQuery::projectedPointOnPolyInTile(const dtMeshTile* tile, cons
 	dtReal verts[DT_VERTS_PER_POLYGON * 3];
 	const int nv = poly->vertCount;
 	for (int i = 0; i < nv; ++i)
-		dtVcopy(&verts[i * 3], &tile->verts[poly->verts[i] * 3]);
+		tile->verts[poly->verts[i]].toWorld(tile->header->bmin, &verts[i * 3]); //@UE
 
 	// copy source to output, just to have any valid information there
 	dtVcopy(projected, pos);
@@ -1036,11 +1067,21 @@ dtStatus dtNavMeshQuery::getPolyHeight(dtPolyRef ref, const dtReal* pos, dtReal*
 	
 	if (poly->getType() == DT_POLYTYPE_OFFMESH_POINT)
 	{
-		const dtReal* v0 = &tile->verts[poly->verts[0]*3];
-		const dtReal* v1 = &tile->verts[poly->verts[1]*3];
+		dtReal v0[3], v1[3]; //@UE
+		tile->verts[poly->verts[0]].toWorld(tile->header->bmin, v0); //@UE
+		tile->verts[poly->verts[1]].toWorld(tile->header->bmin, v1); //@UE
 		const dtReal d0 = dtVdist(pos, v0);
 		const dtReal d1 = dtVdist(pos, v1);
-		const dtReal u = d0 / (d0+d1);
+
+		//@UE BEGIN
+		dtReal u;
+		const dtReal sum = d0 + d1;
+		if (sum != 0)
+			u = d0 / sum;
+		else
+			u = 0;
+		//@UE END
+
 		if (height)
 			*height = v0[1] + (v1[1] - v0[1]) * u;
 		return DT_SUCCESS;
@@ -1049,19 +1090,18 @@ dtStatus dtNavMeshQuery::getPolyHeight(dtPolyRef ref, const dtReal* pos, dtReal*
 #if WITH_NAVMESH_SEGMENT_LINKS
 	else if (poly->getType() == DT_POLYTYPE_OFFMESH_SEGMENT)
 	{
+		dtReal sv0[3], sv1[3], sv2[3], sv3[3]; //@UE
+		tile->verts[poly->verts[0]].toWorld(tile->header->bmin, sv0); //@UE
+		tile->verts[poly->verts[1]].toWorld(tile->header->bmin, sv1); //@UE
+		tile->verts[poly->verts[2]].toWorld(tile->header->bmin, sv2); //@UE
+		tile->verts[poly->verts[3]].toWorld(tile->header->bmin, sv3); //@UE
 		dtReal h;
-		if (dtClosestHeightPointTriangle(pos,
-			&tile->verts[poly->verts[0]*3],
-			&tile->verts[poly->verts[2]*3],
-			&tile->verts[poly->verts[1]*3], h))
+		if (dtClosestHeightPointTriangle(pos, sv0, sv2, sv1, h)) //@UE
 		{
 			if (height) *height = h;
 			return DT_SUCCESS;
 		}
-		else if (dtClosestHeightPointTriangle(pos,
-			&tile->verts[poly->verts[1]*3],
-			&tile->verts[poly->verts[2]*3],
-			&tile->verts[poly->verts[3]*3], h))
+		else if (dtClosestHeightPointTriangle(pos, sv1, sv2, sv3, h)) //@UE
 		{
 			if (height) *height = h;
 			return DT_SUCCESS;
@@ -1076,13 +1116,15 @@ dtStatus dtNavMeshQuery::getPolyHeight(dtPolyRef ref, const dtReal* pos, dtReal*
 		for (int j = 0; j < pd->triCount; ++j)
 		{
 			const unsigned char* t = &tile->detailTris[(pd->triBase+j)*4];
+			dtReal vBuf[3][3]; //@UE
 			const dtReal* v[3];
 			for (int k = 0; k < 3; ++k)
 			{
 				if (t[k] < poly->vertCount)
-					v[k] = &tile->verts[poly->verts[t[k]]*3];
+					tile->verts[poly->verts[t[k]]].toWorld(tile->header->bmin, vBuf[k]); //@UE
 				else
-					v[k] = &tile->detailVerts[(pd->vertBase+(t[k]-poly->vertCount))*3];
+					tile->detailVerts[pd->vertBase+(t[k]-poly->vertCount)].toWorld(tile->header->bmin, vBuf[k]); //@UE
+				v[k] = vBuf[k]; //@UE
 			}
 			dtReal h;
 			if (dtClosestHeightPointTriangle(pos, v[0], v[1], v[2], h))
@@ -1494,12 +1536,13 @@ int dtNavMeshQuery::queryPolygonsInTile(const dtMeshTile* tile, const dtReal* qm
 			if (!filter->passFilter(ref, tile, p) || !passLinkFilter(tile, i))
 				continue;
 			// Calc polygon bounds.
-			const dtReal* v = &tile->verts[p->verts[0]*3];
+			dtReal v[3]; //@UE
+			tile->verts[p->verts[0]].toWorld(tile->header->bmin, v); //@UE
 			dtVcopy(bmin, v);
 			dtVcopy(bmax, v);
 			for (int j = 1; j < p->vertCount; ++j)
 			{
-				v = &tile->verts[p->verts[j]*3];
+				tile->verts[p->verts[j]].toWorld(tile->header->bmin, v); //@UE
 				dtVmin(bmin, v);
 				dtVmax(bmax, v);
 			}
@@ -2954,7 +2997,7 @@ dtStatus dtNavMeshQuery::moveAlongSurface(dtPolyRef startRef, const dtReal* star
 		// Collect vertices.
 		const int nverts = curPoly->vertCount;
 		for (int i = 0; i < nverts; ++i)
-			dtVcopy(&verts[i*3], &curTile->verts[curPoly->verts[i]*3]);
+			curTile->verts[curPoly->verts[i]].toWorld(curTile->header->bmin, &verts[i*3]); //@UE
 		
 		// If target is inside the poly, stop search.
 		if (dtPointInPolygon(endPos, verts, nverts))
@@ -3155,8 +3198,8 @@ dtStatus dtNavMeshQuery::getPortalPoints(dtPolyRef from, const dtPoly* fromPoly,
 			if (testLink.ref == to)
 			{
 				const int v = testLink.edge;
-				dtVcopy(left, &fromTile->verts[fromPoly->verts[v]*3]);
-				dtVcopy(right, &fromTile->verts[fromPoly->verts[v]*3]);
+				fromTile->verts[fromPoly->verts[v]].toWorld(fromTile->header->bmin, left); //@UE
+				fromTile->verts[fromPoly->verts[v]].toWorld(fromTile->header->bmin, right); //@UE
 				return DT_SUCCESS;
 			}
 		}
@@ -3176,8 +3219,8 @@ dtStatus dtNavMeshQuery::getPortalPoints(dtPolyRef from, const dtPoly* fromPoly,
 			if (testLink.ref == to)
 			{
 				const int v = testLink.edge * 2;
-				dtVcopy(left, &fromTile->verts[fromPoly->verts[v+0]*3]);
-				dtVcopy(right, &fromTile->verts[fromPoly->verts[v+1]*3]);
+				fromTile->verts[fromPoly->verts[v+0]].toWorld(fromTile->header->bmin, left); //@UE
+				fromTile->verts[fromPoly->verts[v+1]].toWorld(fromTile->header->bmin, right); //@UE
 				return DT_SUCCESS;
 			}
 		}
@@ -3198,8 +3241,8 @@ dtStatus dtNavMeshQuery::getPortalPoints(dtPolyRef from, const dtPoly* fromPoly,
 			if (testLink.ref == from)
 			{
 				const int v = testLink.edge;
-				dtVcopy(left, &toTile->verts[toPoly->verts[v]*3]);
-				dtVcopy(right, &toTile->verts[toPoly->verts[v]*3]);
+				toTile->verts[toPoly->verts[v]].toWorld(toTile->header->bmin, left); //@UE
+				toTile->verts[toPoly->verts[v]].toWorld(toTile->header->bmin, right); //@UE
 				return DT_SUCCESS;
 			}
 		}
@@ -3218,8 +3261,8 @@ dtStatus dtNavMeshQuery::getPortalPoints(dtPolyRef from, const dtPoly* fromPoly,
 			if (testLink.ref == from)
 			{
 				const int v = testLink.edge * 2;
-				dtVcopy(left, &toTile->verts[toPoly->verts[v+0]*3]);
-				dtVcopy(right, &toTile->verts[toPoly->verts[v+1]*3]);
+				toTile->verts[toPoly->verts[v+0]].toWorld(toTile->header->bmin, left); //@UE
+				toTile->verts[toPoly->verts[v+1]].toWorld(toTile->header->bmin, right); //@UE
 				return DT_SUCCESS;
 			}
 		}
@@ -3231,9 +3274,12 @@ dtStatus dtNavMeshQuery::getPortalPoints(dtPolyRef from, const dtPoly* fromPoly,
 	// Find portal vertices.
 	const int v0 = fromPoly->verts[link->edge];
 	const int v1 = fromPoly->verts[(link->edge+1) % (int)fromPoly->vertCount];
-	dtVcopy(left, &fromTile->verts[v0*3]);
-	dtVcopy(right, &fromTile->verts[v1*3]);
-	
+	dtReal v0pos[3], v1pos[3]; //@UE
+	fromTile->verts[v0].toWorld(fromTile->header->bmin, v0pos); //@UE
+	fromTile->verts[v1].toWorld(fromTile->header->bmin, v1pos); //@UE
+	dtVcopy(left, v0pos); //@UE
+	dtVcopy(right, v1pos); //@UE
+
 	// If the link is at tile boundary, dtClamp the vertices to
 	// the link width.
 //@UE BEGIN
@@ -3246,8 +3292,8 @@ dtStatus dtNavMeshQuery::getPortalPoints(dtPolyRef from, const dtPoly* fromPoly,
 			const dtReal s = dtReal(1.)/255.0f;
 			const dtReal tmin = link->bmin*s;
 			const dtReal tmax = link->bmax*s;
-			dtVlerp(left, &fromTile->verts[v0*3], &fromTile->verts[v1*3], tmin);
-			dtVlerp(right, &fromTile->verts[v0*3], &fromTile->verts[v1*3], tmax);
+			dtVlerp(left, v0pos, v1pos, tmin); //@UE
+			dtVlerp(right, v0pos, v1pos, tmax); //@UE
 		}
 	}
 	
@@ -3382,7 +3428,7 @@ dtStatus dtNavMeshQuery::raycast(dtPolyRef startRef, const dtReal* startPos, con
 		int nv = 0;
 		for (int i = 0; i < (int)poly->vertCount; ++i)
 		{
-			dtVcopy(&verts[nv*3], &tile->verts[poly->verts[i]*3]);
+			tile->verts[poly->verts[i]].toWorld(tile->header->bmin, &verts[nv*3]); //@UE
 			nv++;
 		}
 		
@@ -3470,8 +3516,9 @@ dtStatus dtNavMeshQuery::raycast(dtPolyRef startRef, const dtReal* startPos, con
 			const int v0 = poly->verts[link.edge];
 			CA_SUPPRESS(6385);
 			const int v1 = poly->verts[(link.edge+1) % poly->vertCount];
-			const dtReal* left = &tile->verts[v0*3];
-			const dtReal* right = &tile->verts[v1*3];
+			dtReal left[3], right[3]; //@UE
+			tile->verts[v0].toWorld(tile->header->bmin, left); //@UE
+			tile->verts[v1].toWorld(tile->header->bmin, right); //@UE
 			
 //@UE BEGIN
 			// strip off additional flags
@@ -4230,7 +4277,7 @@ dtStatus dtNavMeshQuery::findLocalNeighbourhood(dtPolyRef startRef, const dtReal
 			// Collect vertices of the neighbour poly.
 			const int npa = neighbourPoly->vertCount;
 			for (int neighbourPolyVertIndex = 0; neighbourPolyVertIndex < npa; ++neighbourPolyVertIndex)
-				dtVcopy(&pa[neighbourPolyVertIndex * 3], &neighbourTile->verts[neighbourPoly->verts[neighbourPolyVertIndex] * 3]);
+				neighbourTile->verts[neighbourPoly->verts[neighbourPolyVertIndex]].toWorld(neighbourTile->header->bmin, &pa[neighbourPolyVertIndex * 3]); //@UE
 			
 			bool overlap = false;
 			for (int j = 0; j < n; ++j)
@@ -4262,7 +4309,7 @@ dtStatus dtNavMeshQuery::findLocalNeighbourhood(dtPolyRef startRef, const dtReal
 				// Get vertices and test overlap
 				const int npb = pastPoly->vertCount;
 				for (int pastPolyVertIndex = 0; pastPolyVertIndex < npb; ++pastPolyVertIndex)
-					dtVcopy(&pb[pastPolyVertIndex * 3], &pastTile->verts[pastPoly->verts[pastPolyVertIndex] * 3]);
+					pastTile->verts[pastPoly->verts[pastPolyVertIndex]].toWorld(pastTile->header->bmin, &pb[pastPolyVertIndex * 3]); //@UE
 				
 				if (dtOverlapPolyPoly2D(pa,npa, pb,npb))
 				{
@@ -4408,11 +4455,9 @@ dtStatus dtNavMeshQuery::getPolyWallSegments(dtPolyRef ref, const dtQueryFilter*
 			
 			if (n < maxSegments)
 			{
-				const dtReal* vj = &tile->verts[poly->verts[j]*3];
-				const dtReal* vi = &tile->verts[poly->verts[i]*3];
 				dtReal* seg = &segmentVerts[n*6];
-				dtVcopy(seg+0, vj);
-				dtVcopy(seg+3, vi);
+				tile->verts[poly->verts[j]].toWorld(tile->header->bmin, seg+0); //@UE
+				tile->verts[poly->verts[i]].toWorld(tile->header->bmin, seg+3); //@UE
 				if (segmentRefs)
 					segmentRefs[n] = neiRef;
 				n++;
@@ -4430,8 +4475,9 @@ dtStatus dtNavMeshQuery::getPolyWallSegments(dtPolyRef ref, const dtQueryFilter*
 		insertInterval(ints, nints, MAX_INTERVAL, 255, 256, 0);
 		
 		// Store segments.
-		const dtReal* vj = &tile->verts[poly->verts[j]*3];
-		const dtReal* vi = &tile->verts[poly->verts[i]*3];
+		dtReal vj[3], vi[3]; //@UE
+		tile->verts[poly->verts[j]].toWorld(tile->header->bmin, vj); //@UE
+		tile->verts[poly->verts[i]].toWorld(tile->header->bmin, vi); //@UE
 		for (int k = 1; k < nints; ++k)
 		{
 			// Portal segment.
@@ -4485,8 +4531,9 @@ dtStatus dtNavMeshQuery::getPolyWallSegments(dtPolyRef ref, const dtQueryFilter*
 
 static bool isEdgeInRadius(const dtMeshTile* tile, const dtPoly* poly, int edge, const dtReal* centerPos, const dtReal radiusSqr)
 {
-	const dtReal* va = &tile->verts[poly->verts[edge] * 3];
-	const dtReal* vb = &tile->verts[poly->verts[(edge + 1) % poly->vertCount] * 3];
+	dtReal va[3], vb[3]; //@UE
+	tile->verts[poly->verts[edge]].toWorld(tile->header->bmin, va); //@UE
+	tile->verts[poly->verts[(edge + 1) % poly->vertCount]].toWorld(tile->header->bmin, vb); //@UE
 
 	dtReal tseg;
 	const dtReal distSqr = dtDistancePtSegSqr2D(centerPos, va, vb, tseg);
@@ -4518,18 +4565,19 @@ static void storeWallSegment(const dtNavMesh* nav, const dtMeshTile* tile, const
 		return;
 	}
 
-	const int vertIndexA = poly->verts[edge] * 3;
-	const int vertIndexB = poly->verts[tempEdgeIndexB] * 3;
+	const int vertIndexA = poly->verts[edge]; //@UE
+	const int vertIndexB = poly->verts[tempEdgeIndexB]; //@UE
 
-	const int tileVertArraySize = tile->header->vertCount * 3;
+	const int tileVertArraySize = tile->header->vertCount; //@UE
 	if (!ensureMsgf(vertIndexA < tileVertArraySize && vertIndexB < tileVertArraySize,
 		TEXT("Invalid vert index (vertIndexA=%i, vertIndexB=%i, vertCount=%i)"), vertIndexA, vertIndexB, tile->header->vertCount))
 	{
 		return;
 	}
 
-	const dtReal* va = &tile->verts[vertIndexA];
-	const dtReal* vb = &tile->verts[vertIndexB];
+	dtReal va[3], vb[3]; //@UE
+	tile->verts[vertIndexA].toWorld(tile->header->bmin, va); //@UE
+	tile->verts[vertIndexB].toWorld(tile->header->bmin, vb); //@UE
 
 	const int32 wall0Offset = (*resultCount * 6) + 0;
 	const int32 wall1Offset = (*resultCount * 6) + 3;
@@ -4564,8 +4612,9 @@ static void storeWallSegment(const dtNavMesh* nav, const dtMeshTile* tile, const
 
 		if (neiEdge != -1)
 		{
-			const dtReal* va2 = &neiTile->verts[neiPoly->verts[neiEdge] * 3];
-			const dtReal* vb2 = &neiTile->verts[neiPoly->verts[(neiEdge + 1) % neiPoly->vertCount] * 3];
+			dtReal va2[3], vb2[3]; //@UE
+			neiTile->verts[neiPoly->verts[neiEdge]].toWorld(neiTile->header->bmin, va2); //@UE
+			neiTile->verts[neiPoly->verts[(neiEdge + 1) % neiPoly->vertCount]].toWorld(neiTile->header->bmin, vb2); //@UE
 
 			// Project and clip segment va2-vb2 on va-vb
 			dtReal seg[3], diffA[3], diffB[3], clippedA[3], clippedB[3];
@@ -5237,10 +5286,11 @@ dtStatus dtNavMeshQuery::findDistanceToWall(dtPolyRef startRef, const dtReal* ce
 									constexpr dtReal s = 1.0/255.0;
 									const dtReal tmin = link.bmin*s;
 									const dtReal tmax = link.bmax*s;
-									const dtReal* vj = &bestTile->verts[bestPoly->verts[j]*3];
-									const dtReal* vi = &bestTile->verts[bestPoly->verts[i]*3];
-									dtVlerp(portalLeft,  vj, vi, tmin);
-									dtVlerp(portalRight, vj, vi, tmax);
+									dtReal vjTmp[3], viTmp[3]; //@UE
+									bestTile->verts[bestPoly->verts[j]].toWorld(bestTile->header->bmin, vjTmp); //@UE
+									bestTile->verts[bestPoly->verts[i]].toWorld(bestTile->header->bmin, viTmp); //@UE
+									dtVlerp(portalLeft,  vjTmp, viTmp, tmin); //@UE
+									dtVlerp(portalRight, vjTmp, viTmp, tmax); //@UE
 								}
 							}
 						}
@@ -5261,12 +5311,15 @@ dtStatus dtNavMeshQuery::findDistanceToWall(dtPolyRef startRef, const dtReal* ce
 			}
 			
 			// Calc distance to the edge.
-			dtReal* vj;
-			dtReal* vi;
+			dtReal vjBuf[3], viBuf[3]; //@UE
+			const dtReal* vj; //@UE
+			const dtReal* vi; //@UE
 			if (!isPortal)
 			{
-				vj = &bestTile->verts[bestPoly->verts[j]*3];
-				vi = &bestTile->verts[bestPoly->verts[i]*3];
+				bestTile->verts[bestPoly->verts[j]].toWorld(bestTile->header->bmin, vjBuf); //@UE
+				bestTile->verts[bestPoly->verts[i]].toWorld(bestTile->header->bmin, viBuf); //@UE
+				vj = vjBuf; //@UE
+				vi = viBuf; //@UE
 			}
 			else
 			{
@@ -5311,9 +5364,10 @@ dtStatus dtNavMeshQuery::findDistanceToWall(dtPolyRef startRef, const dtReal* ce
 				continue;
 			
 			// Calc distance to the edge.
-			const dtReal* va = &bestTile->verts[bestPoly->verts[link.edge]*3];
+			dtReal va[3], vb[3]; //@UE
+			bestTile->verts[bestPoly->verts[link.edge]].toWorld(bestTile->header->bmin, va); //@UE
 			CA_SUPPRESS(6385);
-			const dtReal* vb = &bestTile->verts[bestPoly->verts[(link.edge+1) % bestPoly->vertCount]*3];
+			bestTile->verts[bestPoly->verts[(link.edge+1) % bestPoly->vertCount]].toWorld(bestTile->header->bmin, vb); //@UE
 			dtReal tseg;
 			const dtReal distSqr = dtDistancePtSegSqr2D(centerPos, va, vb, tseg);
 			
diff --git a/Engine/Source/Runtime/Navmesh/Public/Detour/DetourNavMesh.h b/Engine/Source/Runtime/Navmesh/Public/Detour/DetourNavMesh.h
--- a/Engine/Source/Runtime/Navmesh/Public/Detour/DetourNavMesh.h
+++ b/Engine/Source/Runtime/Navmesh/Public/Detour/DetourNavMesh.h
@@ -26,6 +26,7 @@
 #include "Detour/DetourAlloc.h"
 #include "Detour/DetourStatus.h"
 #include "DetourLargeWorldCoordinates.h"
+#include "Detour/DetourAssert.h"
 #include "Logging/LogMacros.h"
 
 NAVMESH_API DECLARE_LOG_CATEGORY_EXTERN(LogDetour, Log, All);
@@ -92,7 +93,7 @@ static const int DT_VERTS_PER_POLYGON = 6;
 //static const int DT_NAVMESH_MAGIC = 'D'<<24 | 'N'<<16 | 'A'<<8 | 'V';
 
 /// A version number used to detect compatibility of navigation tile data.
-static const int DT_NAVMESH_VERSION = 7;
+static const int DT_NAVMESH_VERSION = 8;
 
 /// A magic number used to detect the compatibility of navigation tile states.
 static const int DT_NAVMESH_STATE_MAGIC = 'D'<<24 | 'N'<<16 | 'M'<<8 | 'S';
@@ -201,8 +202,47 @@ enum dtPolyTypes
 	//@UE END
 };
 
+//@UE BEGIN
+/// Tile-local vertex: stores a vertex as offsets from the tile's bmin.
+/// Use toWorld() to reconstruct the world-space position, and setFromWorld() to write back.
+struct dtTileVert
+{
+	float x, y, z;
+
+	/// Encode from world position
+	inline void setFromWorld(const dtReal* bmin, const dtReal* worldPos)
+	{
+		const dtReal dx = worldPos[0] - bmin[0];
+		const dtReal dy = worldPos[1] - bmin[1];
+		const dtReal dz = worldPos[2] - bmin[2];
+
+		// Verify that tile-local offsets are reasonable for float storage.
+
+		// Maximum allowed tile-local vertex offset before precision becomes problematic.
+		// Typical tiles are <10,000 units. This threshold (100x that) catches bugs like:
+		// - Vertex assigned to wrong tile
+		// - Wrong bmin passed to setFromWorld()
+		// - Coordinate system errors
+		// Values this large indicate real bugs, not just edge cases.
+		constexpr dtReal DT_TILE_VERT_MAX_OFFSET = 1000000.0;
+		dtAssert((dx < 0 ? -dx : dx) < DT_TILE_VERT_MAX_OFFSET && (dy < 0 ? -dy : dy) < DT_TILE_VERT_MAX_OFFSET && (dz < 0 ? -dz : dz) < DT_TILE_VERT_MAX_OFFSET);
+
+		x = (float)dx;
+		y = (float)dy;
+		z = (float)dz;
+	}
+	
+	/// Decode to world position
+	inline void toWorld(const dtReal* bmin, dtReal* out) const
+	{
+		out[0] = (dtReal)x + bmin[0];
+		out[1] = (dtReal)y + bmin[1];
+		out[2] = (dtReal)z + bmin[2];
+	}
+};
+//@UE END
 
-/// Defines a polyogn within a dtMeshTile object.
+/// Defines a polygon within a dtMeshTile object.
 /// @ingroup detour
 struct dtPoly
 {
@@ -432,12 +472,12 @@ struct dtMeshTile
 	unsigned int linksFreeList;			///< Index to the next free link.
 	dtMeshHeader* header;				///< The tile header.
 	dtPoly* polys;						///< The tile polygons. [Size: dtMeshHeader::polyCount]
-	dtReal* verts;						///< The tile vertices. [Size: dtMeshHeader::vertCount (3 real per vertex)]
+	dtTileVert* verts;					///< The tile vertices. Local offsets from dtMeshHeader::bmin. [Size: dtMeshHeader::vertCount] //@UE
 	dtLink* links;						///< The tile links. [Size: dtMeshHeader::maxLinkCount]
 	dtPolyDetail* detailMeshes;			///< The tile's detail sub-meshes. [Size: dtMeshHeader::detailMeshCount]
-	
-	/// The detail mesh's unique vertices. [(x, y, z) * dtMeshHeader::detailVertCount]
-	dtReal* detailVerts;
+
+	/// The detail mesh's unique vertices. [Size: dtMeshHeader::detailVertCount] Local offsets from dtMeshHeader::bmin. //@UE
+	dtTileVert* detailVerts; //@UE
 
 	/// The detail mesh's triangles. [(vertA, vertB, vertC) * dtMeshHeader::detailTriCount]
 	unsigned char* detailTris;	
