Fix tile ordering dependency in connectDeferredOffMeshLinks

connectDeferredOffMeshLinks processed tiles in a single pass, which meant
a tile's off-mesh connections could be skipped if its neighbour hadn't been
initialized yet. The result was intermittent missing off-mesh links, dependent
on tile array ordering.

Fixed by splitting into two passes: first initialize all tiles, then connect them.

#Synced-CL 52421151

diff --git a/Engine/Source/Runtime/Navmesh/Private/Detour/DetourNavMesh.cpp b/Engine/Source/Runtime/Navmesh/Private/Detour/DetourNavMesh.cpp
--- a/Engine/Source/Runtime/Navmesh/Private/Detour/DetourNavMesh.cpp
+++ b/Engine/Source/Runtime/Navmesh/Private/Detour/DetourNavMesh.cpp
@@ -2149,6 +2149,18 @@ void dtNavMesh::connectDeferredOffMeshLinks()
 
 	ReadTilesHelper TileArray;
 
+	// Pass 1: Initialize all tiles' off-mesh connections first.
+	// baseOffMeshLinks() sets poly->firstLink on each tile's off-mesh polys.
+	// connectExtOffMeshLinks() guards on firstLink != DT_NULL_LINK, so all tiles
+	// must be initialized before any cross-tile connections are attempted.
+	for (int i = 0; i < m_maxTiles; ++i)
+	{
+		dtMeshTile* tile = &m_tiles[i];
+		if (!tile->header) continue;
+		baseOffMeshLinks(tile);
+	}
+
+	// Pass 2: Connect external links now that firstLink is valid on all tiles.
 	for (int i = 0; i < m_maxTiles; ++i)
 	{
 		dtMeshTile* tile = &m_tiles[i];
@@ -2162,9 +2174,6 @@ void dtNavMesh::connectDeferredOffMeshLinks()
 #else
 		const bool bHasClusters = false;
 #endif // WITH_NAVMESH_CLUSTER_LINKS
-		
-		// Snap start points (vertex 0) of this tile's off-mesh connections
-		baseOffMeshLinks(tile);
 
 		// Connect layers in same column — landing on this tile
 		int nneis = getTileCountAt(tx, ty);
