[Navigation] Optimize dtLink memory allocation (mainly internal links) with minimized base pool
Reduces navmesh link memory by ~30 Mb on large map using optimized base pool sizing
- Changed base pool allocation to remove waste
- Added bMinimizeLinkPoolSize property (true by default). When false, reverts to larger base pool.
- Added link utilization instrumentation to DumpNavMeshMemory()

#Synced-CL 52469485

diff --git a/Engine/Source/Runtime/Engine/Public/EngineStats.h b/Engine/Source/Runtime/Engine/Public/EngineStats.h
--- a/Engine/Source/Runtime/Engine/Public/EngineStats.h
+++ b/Engine/Source/Runtime/Engine/Public/EngineStats.h
@@ -169,7 +169,7 @@ DECLARE_MEMORY_STAT_EXTERN(TEXT("Detour PERM_PATH_CORRIDOR"), STAT_Navigation_De
 DECLARE_MEMORY_STAT_EXTERN(TEXT("Detour PERM_PATH_QUEUE"), STAT_Navigation_DetourPERM_PATH_QUEUE, STATGROUP_Navigation, );
 DECLARE_MEMORY_STAT_EXTERN(TEXT("Detour PERM_PROXY_GRID"), STAT_Navigation_DetourPERM_PROXY_GRID, STATGROUP_Navigation, );
 DECLARE_MEMORY_STAT_EXTERN(TEXT("Detour PERM_TILE_DATA"), STAT_Navigation_DetourPERM_TILE_DATA, STATGROUP_Navigation, );
-DECLARE_MEMORY_STAT_EXTERN(TEXT("Detour PERM_TILE_DYNLINK_OFFMESH"), STAT_Navigation_DetourPERM_TILE_DYNLINK_OFFMESH, STATGROUP_Navigation, );
+DECLARE_MEMORY_STAT_EXTERN(TEXT("Detour PERM_TILE_DYNLINK"), STAT_Navigation_DetourPERM_TILE_DYNLINK, STATGROUP_Navigation, );
 DECLARE_MEMORY_STAT_EXTERN(TEXT("Detour PERM_TILE_DYNLINK_CLUSTER"), STAT_Navigation_DetourPERM_TILE_DYNLINK_CLUSTER, STATGROUP_Navigation, );
 DECLARE_MEMORY_STAT_EXTERN(TEXT("Detour PERM_TILES"), STAT_Navigation_DetourPERM_TILES, STATGROUP_Navigation, );
 DECLARE_MEMORY_STAT_EXTERN(TEXT("Detour PERM_TILE_LINK_BUILDER"), STAT_Navigation_DetourPERM_TILE_LINK_BUILDER, STATGROUP_Navigation, );
diff --git a/Engine/Source/Runtime/NavigationSystem/Private/NavMesh/PImplRecastNavMesh.cpp b/Engine/Source/Runtime/NavigationSystem/Private/NavMesh/PImplRecastNavMesh.cpp
--- a/Engine/Source/Runtime/NavigationSystem/Private/NavMesh/PImplRecastNavMesh.cpp
+++ b/Engine/Source/Runtime/NavigationSystem/Private/NavMesh/PImplRecastNavMesh.cpp
@@ -98,8 +98,8 @@ static void* DetourMalloc(int Size, dtAllocHint Hint)
 		INC_MEMORY_STAT_BY(STAT_Navigation_DetourPERM_TILE_DATA, ActualSize);
 		break;
 
-	case DT_ALLOC_PERM_TILE_DYNLINK_OFFMESH:
-		INC_MEMORY_STAT_BY(STAT_Navigation_DetourPERM_TILE_DYNLINK_OFFMESH, ActualSize);
+	case DT_ALLOC_PERM_TILE_DYNLINK:
+		INC_MEMORY_STAT_BY(STAT_Navigation_DetourPERM_TILE_DYNLINK, ActualSize);
 		break;
 
 	case DT_ALLOC_PERM_TILE_DYNLINK_CLUSTER:
@@ -188,8 +188,8 @@ static void DetourFree(void* Original, dtAllocHint Hint)
 		DEC_MEMORY_STAT_BY(STAT_Navigation_DetourPERM_TILE_DATA, Size);
 		break;
 
-	case DT_ALLOC_PERM_TILE_DYNLINK_OFFMESH:
-		DEC_MEMORY_STAT_BY(STAT_Navigation_DetourPERM_TILE_DYNLINK_OFFMESH, Size);
+	case DT_ALLOC_PERM_TILE_DYNLINK:
+		DEC_MEMORY_STAT_BY(STAT_Navigation_DetourPERM_TILE_DYNLINK, Size);
 		break;
 
 	case DT_ALLOC_PERM_TILE_DYNLINK_CLUSTER:
diff --git a/Engine/Source/Runtime/NavigationSystem/Private/NavMesh/RecastNavMesh.cpp b/Engine/Source/Runtime/NavigationSystem/Private/NavMesh/RecastNavMesh.cpp
--- a/Engine/Source/Runtime/NavigationSystem/Private/NavMesh/RecastNavMesh.cpp
+++ b/Engine/Source/Runtime/NavigationSystem/Private/NavMesh/RecastNavMesh.cpp
@@ -512,6 +512,7 @@ ARecastNavMesh::ARecastNavMesh(const FObjectInitializer& ObjectInitializer)
 	, DrawOffset(10.f)
 	, TilePoolSize(1024)
 	, AverageLayersPerTile(3)
+	, bMinimizeLinkPoolSize(true)
 	, MaxVerticalMergeError(INT_MAX) // By default, ignore vertical error
 	, MaxSimplificationError(1.3f)	// from RecastDemo
 	, SimplificationElevationRatio(0.f)	// By default, ignore contour simplification from elevation
@@ -3313,6 +3314,22 @@ void ARecastNavMesh::DumpDetailedMemoryStats() const
 	int32 MinPolysPerTile = MAX_int32;
 	int32 MaxPolysPerTile = 0;
 
+	// Link utilization tracking
+	uint64 TotalBaseLinksAllocated = 0;
+	uint64 TotalLinksUsed = 0;
+	int32 MaxPreAllocatedLinksInAnyTile = 0;
+	uint64 TotalLinksPerPolySum = 0;
+	int32 TotalPolyCountForAvg = 0;
+
+	// Link utilization distribution buckets (0-10%, 10-20%, ..., 90-100%)
+	static constexpr int32 UtilBuckets = 10;
+	int32 UtilBucketCount[UtilBuckets] = {0};
+
+	// Dynamic allocation tracking
+	int32 TilesWithDynamic = 0;
+	uint64 TotalDynamicLinks = 0;
+	int32 MaxDynamicInTile = 0;
+
 	// Histogram: poly count per tile
 	// Buckets: [1-2], [3-4], [5-8], [9-16], [17-32], [33-64], [65-128], [129-256], [257+]
 	static constexpr int32 NumBuckets = 9;
@@ -3360,6 +3377,51 @@ void ARecastNavMesh::DumpDetailedMemoryStats() const
 		TotalDetailTriCount += Header->detailTriCount;
 		TotalOffMeshConCount += Header->offMeshConCount;
 
+		// Count actually used links by traversing polygon link chains
+		int32 UsedBaseLinks = 0;
+		int32 UsedDynamicLinks = 0;
+		for (int32 PolyIdx = 0; PolyIdx < Header->polyCount; ++PolyIdx)
+		{
+			const dtPoly* Poly = &Tile->polys[PolyIdx];
+			unsigned int LinkIdx = Poly->firstLink;
+			while (LinkIdx != DT_NULL_LINK)
+			{
+				if (LinkIdx < (unsigned int)Header->maxLinkCount)
+				{
+					UsedBaseLinks++;
+				}
+				else
+				{
+					UsedDynamicLinks++;
+				}
+				const dtLink& Link = NavMesh->getLink(Tile, LinkIdx);
+				LinkIdx = Link.next;
+			}
+		}
+
+		const int32 UsedLinksInTile = UsedBaseLinks + UsedDynamicLinks;
+		TotalBaseLinksAllocated += Header->maxLinkCount;
+		TotalLinksUsed += UsedLinksInTile;
+		MaxPreAllocatedLinksInAnyTile = FMath::Max(MaxPreAllocatedLinksInAnyTile, (int32)Header->maxLinkCount);
+
+		// Track dynamic allocations
+		TotalDynamicLinks += UsedDynamicLinks;
+		if (UsedDynamicLinks > 0)
+		{
+			TilesWithDynamic++;
+			MaxDynamicInTile = FMath::Max(MaxDynamicInTile, UsedDynamicLinks);
+		}
+
+		// Links per polygon stats
+		TotalLinksPerPolySum += UsedLinksInTile;
+		TotalPolyCountForAvg += Header->polyCount;
+
+		// Utilization bucket (base pool only)
+		const float Utilization = Header->maxLinkCount > 0
+			? (float)UsedBaseLinks / Header->maxLinkCount : 0.0f;
+		const int32 BucketIdx = FMath::Clamp((int32)(Utilization * 10.0f), 0, UtilBuckets - 1);
+		UtilBucketCount[BucketIdx]++;
+
 		if (Header->bvNodeCount > 0)
 		{
 			TilesWithBvTree++;
@@ -3415,7 +3477,7 @@ void ARecastNavMesh::DumpDetailedMemoryStats() const
 	UE_LOGF(LogNavigation, Display, "--- Overall Memory ---");
 	UE_LOGF(LogNavigation, Display, "  Tile data total:     %10llu bytes  (%.2f MB)", TotalTileDataBytes, (float)TotalTileDataBytes * ToMB);
 	UE_LOGF(LogNavigation, Display, "  Tile pool structs:   %10llu bytes  (%.2f MB)", TilePoolBytes, (float)TilePoolBytes * ToMB);
-	UE_LOGF(LogNavigation, Display, "  Lookup table:        %10llu bytes  (%.2f KB)", LookupTableBytes, (float)LookupTableBytes * ToKB);
+	UE_LOGF(LogNavigation, Display, "  Lookup table:        %10llu bytes  (%.2f MB)", LookupTableBytes, (float)LookupTableBytes * ToMB);
 	UE_LOGF(LogNavigation, Display, "  TOTAL:               %10llu bytes  (%.2f MB)",
 		(TotalTileDataBytes + TilePoolBytes + LookupTableBytes),
 		(float)(TotalTileDataBytes + TilePoolBytes + LookupTableBytes) * ToMB);
@@ -3485,6 +3547,55 @@ void ARecastNavMesh::DumpDetailedMemoryStats() const
 			(float)BucketDetailBytes[Bucket] * ToKB,
 			(float)BucketDetailBytes[Bucket] * ToMB);
 	}
+	UE_LOGF(LogNavigation, Display, "");
+
+	// Link utilization analysis
+	const uint64 TotalAllocated = TotalBaseLinksAllocated + TotalDynamicLinks;
+	UE_LOGF(LogNavigation, Display, "--- Link Utilization Analysis ---");
+	UE_LOGF(LogNavigation, Display, "  Base pool allocated:   %llu links", TotalBaseLinksAllocated);
+	UE_LOGF(LogNavigation, Display, "  Dynamic used:        %llu links", TotalDynamicLinks);
+	UE_LOGF(LogNavigation, Display, "  Total allocated:       %llu links", TotalAllocated);
+	UE_LOGF(LogNavigation, Display, "  Links actually used:   %llu (%.1f%%)", TotalLinksUsed, 100.0f * (float)TotalLinksUsed / (float)FMath::Max((uint64)1, TotalAllocated));
+
+	const uint64 WastedLinks = TotalAllocated > TotalLinksUsed
+		? TotalAllocated - TotalLinksUsed : 0;
+
+	UE_LOGF(LogNavigation, Display, "  Wasted links:          %llu (%.2f MB)",
+		WastedLinks,
+		(float)WastedLinks * (float)sizeof(dtLink) * ToMB);
+	UE_LOGF(LogNavigation, Display, "  Max preallocated links in any tile: %d", MaxPreAllocatedLinksInAnyTile);
+
+	if (TotalPolyCountForAvg > 0)
+	{
+		UE_LOGF(LogNavigation, Display, "  Links per poly:        avg=%.2f",
+			(float)TotalLinksPerPolySum / (float)TotalPolyCountForAvg);
+	}
+
+	UE_LOGF(LogNavigation, Display, "");
+	UE_LOGF(LogNavigation, Display, "--- Tile Utilization Distribution ---");
+	UE_LOGF(LogNavigation, Display, "  Utilization    Tiles     %%");
+	for (int32 i = 0; i < UtilBuckets; ++i)
+	{
+		if (UtilBucketCount[i] > 0)
+		{
+			UE_LOGF(LogNavigation, Display, "  %3d%%-%3d%%     %6d   %5.1f%%",
+				i * 10, (i + 1) * 10,
+				UtilBucketCount[i],
+				100.0f * (float)UtilBucketCount[i] / (float)TotalActiveTiles);
+		}
+	}
+	UE_LOGF(LogNavigation, Display, "");
+
+	// Dynamic allocation statistics
+	UE_LOGF(LogNavigation, Display, "--- Dynamic Allocation Statistics ---");
+	UE_LOGF(LogNavigation, Display, "  Tiles using dynamic: %d (%.1f%%)",
+		TilesWithDynamic, 100.0f * (float)TilesWithDynamic / (float)FMath::Max(TotalActiveTiles, 1));
+	UE_LOGF(LogNavigation, Display, "  Total dynamic links: %llu (%.1f%% of used links)",
+		TotalDynamicLinks,
+		TotalLinksUsed > 0 ? 100.0f * (float)TotalDynamicLinks / (float)TotalLinksUsed : 0.0f);
+	UE_LOGF(LogNavigation, Display, "  Max dynamic in tile: %d links", MaxDynamicInTile);
+	UE_LOGF(LogNavigation, Display, "  Dynamic memory:      %.2f MB",
+		(float)TotalDynamicLinks * (float)sizeof(dtLink) * ToMB);
 
 	UE_LOGF(LogNavigation, Display, "=======================================================================================================");
 	UE_LOGF(LogNavigation, Display, "");
diff --git a/Engine/Source/Runtime/NavigationSystem/Private/NavMesh/RecastNavMeshGenerator.cpp b/Engine/Source/Runtime/NavigationSystem/Private/NavMesh/RecastNavMeshGenerator.cpp
--- a/Engine/Source/Runtime/NavigationSystem/Private/NavMesh/RecastNavMeshGenerator.cpp
+++ b/Engine/Source/Runtime/NavigationSystem/Private/NavMesh/RecastNavMeshGenerator.cpp
@@ -4408,6 +4408,8 @@ bool FRecastTileGenerator::GenerateNavigationDataLayer(FNavMeshBuildContext& Bui
 		const bool bShouldBuildBvTree = TileConfig.bGenerateBVTree && (Params.polyCount > PolyCountThreshold);
 		Params.buildBvTree = bShouldBuildBvTree;
 
+		Params.bMinimizeLinkPoolSize = TileConfig.bMinimizeLinkPoolSize;
+
 #if WITH_NAVMESH_CLUSTER_LINKS
 		Params.clusterCount = IntCastChecked<unsigned short>(GenerationContext.ClusterSet->nclusters);
 		Params.polyClusters = GenerationContext.ClusterSet->polyMap;
@@ -5345,6 +5347,7 @@ void FRecastNavMeshGenerator::ConfigureBuildProperties(FRecastBuildConfig& OutCo
 	OutConfig.bUseExtraTopCellWhenMarkingAreas = DestNavMesh->bUseExtraTopCellWhenMarkingAreas;
 	OutConfig.bFilterLowSpanSequences = DestNavMesh->bFilterLowSpanSequences;
 	OutConfig.bFilterLowSpanFromTileCache = DestNavMesh->bFilterLowSpanFromTileCache;
+	OutConfig.bMinimizeLinkPoolSize = DestNavMesh->bMinimizeLinkPoolSize;
 	if (DestNavMesh->bMarkLowHeightAreas)
 	{
 		OutConfig.walkableHeight = 1;
diff --git a/Engine/Source/Runtime/NavigationSystem/Private/NavigationSystem.cpp b/Engine/Source/Runtime/NavigationSystem/Private/NavigationSystem.cpp
--- a/Engine/Source/Runtime/NavigationSystem/Private/NavigationSystem.cpp
+++ b/Engine/Source/Runtime/NavigationSystem/Private/NavigationSystem.cpp
@@ -153,7 +153,7 @@ DEFINE_STAT(STAT_Navigation_DetourPERM_PATH_CORRIDOR);
 DEFINE_STAT(STAT_Navigation_DetourPERM_PATH_QUEUE);
 DEFINE_STAT(STAT_Navigation_DetourPERM_PROXY_GRID);
 DEFINE_STAT(STAT_Navigation_DetourPERM_TILE_DATA);
-DEFINE_STAT(STAT_Navigation_DetourPERM_TILE_DYNLINK_OFFMESH);
+DEFINE_STAT(STAT_Navigation_DetourPERM_TILE_DYNLINK);
 DEFINE_STAT(STAT_Navigation_DetourPERM_TILE_DYNLINK_CLUSTER);
 DEFINE_STAT(STAT_Navigation_DetourPERM_TILES);
 DEFINE_STAT(STAT_Navigation_DetourPERM_TILE_LINK_BUILDER);
diff --git a/Engine/Source/Runtime/NavigationSystem/Public/NavMesh/RecastNavMesh.h b/Engine/Source/Runtime/NavigationSystem/Public/NavMesh/RecastNavMesh.h
--- a/Engine/Source/Runtime/NavigationSystem/Public/NavMesh/RecastNavMesh.h
+++ b/Engine/Source/Runtime/NavigationSystem/Public/NavMesh/RecastNavMesh.h
@@ -696,6 +696,11 @@ class ARecastNavMesh : public ANavigationData
 	UPROPERTY(EditAnywhere, Category=Generation, config, meta=(ClampMin = "1.0"))
 	float AverageLayersPerTile;
 
+	/** If true (default), minimize base link pool size using optimized allocation, else use previous larger base link pool.
+	 *  Changing this property requires a navmesh rebuild to fully benefit from the change. */
+	UPROPERTY(EditAnywhere, Category=Generation, config, AdvancedDisplay)
+	uint32 bMinimizeLinkPoolSize:1;
+
 	/**
 	 * Note that we are not using _DEPRECATED on the following deprecated properties
 	 * since it prevents the property from being serialized back which can break the
diff --git a/Engine/Source/Runtime/NavigationSystem/Public/NavMesh/RecastNavMeshGenerator.h b/Engine/Source/Runtime/NavigationSystem/Public/NavMesh/RecastNavMeshGenerator.h
--- a/Engine/Source/Runtime/NavigationSystem/Public/NavMesh/RecastNavMeshGenerator.h
+++ b/Engine/Source/Runtime/NavigationSystem/Public/NavMesh/RecastNavMeshGenerator.h
@@ -65,6 +65,8 @@ struct FRecastBuildConfig : public rcConfig
 	uint32 bFilterLowSpanFromTileCache : 1;
 	/** if set, navigation links will be automatically generated */
 	uint32 bGenerateLinks : 1;
+	/** if set, minimize base link pool size, else use previous larger base link pool. */
+	uint32 bMinimizeLinkPoolSize : 1;
 
 	/** region partitioning method used by tile cache */
 	int32 TileCachePartitionType;
@@ -113,6 +115,7 @@ struct FRecastBuildConfig : public rcConfig
 		bFilterLowSpanSequences = false;
 		bFilterLowSpanFromTileCache = false;
 		bGenerateLinks = false;
+		bMinimizeLinkPoolSize = true;
 		MaxPolysPerTile = -1;
 		JumpConfigs.Empty();
 		AgentIndex = 0;
diff --git a/Engine/Source/Runtime/Navmesh/Private/DebugUtils/DetourDebugDraw.cpp b/Engine/Source/Runtime/Navmesh/Private/DebugUtils/DetourDebugDraw.cpp
--- a/Engine/Source/Runtime/Navmesh/Private/DebugUtils/DetourDebugDraw.cpp
+++ b/Engine/Source/Runtime/Navmesh/Private/DebugUtils/DetourDebugDraw.cpp
@@ -27,7 +27,7 @@
 
 inline const dtLink& getLinkHelper(const dtMeshTile* tile, unsigned int link)
 {
-	return (link < (unsigned int)tile->header->maxLinkCount) ? tile->links[link] : tile->dynamicLinksO[link - tile->header->maxLinkCount];
+	return (link < (unsigned int)tile->header->maxLinkCount) ? tile->links[link] : tile->dynamicLinks[link - tile->header->maxLinkCount]; //@UE
 }
 
 static dtReal distancePtLine2d(const dtReal* pt, const dtReal* p, const dtReal* q)
diff --git a/Engine/Source/Runtime/Navmesh/Private/Detour/DetourNavMesh.cpp b/Engine/Source/Runtime/Navmesh/Private/Detour/DetourNavMesh.cpp
--- a/Engine/Source/Runtime/Navmesh/Private/Detour/DetourNavMesh.cpp
+++ b/Engine/Source/Runtime/Navmesh/Private/Detour/DetourNavMesh.cpp
@@ -208,15 +208,14 @@ inline int computeTileHash(int x, int y, const int mask)
 //@UE BEGIN
 enum ELinkAllocationType
 {
-	// allocated in linksFreeList from dtMeshTile (detour)
+	// allocated in linksFreeList from dtMeshTile (detour) first, then to dynamic array
 	CREATE_LINK_PREALLOCATED,
-	// offmesh links will be added in dynamic array (unreal specific)
+	// offmesh links are added in a dynamic array (sharing the space with extra reallocated links) 
 	CREATE_LINK_DYNAMIC_OFFMESH,
 #if WITH_NAVMESH_CLUSTER_LINKS
 	CREATE_LINK_DYNAMIC_CLUSTER,
-#endif // WITH_NAVMESH_SEGMENT_LINKS
+#endif // WITH_NAVMESH_CLUSTER_LINKS
 };
-//@UE END
 
 inline unsigned int allocLink(dtMeshTile* tile, char LinkAllocMode)
 {
@@ -224,29 +223,53 @@ inline unsigned int allocLink(dtMeshTile* tile, char LinkAllocMode)
 
 	if (LinkAllocMode == CREATE_LINK_PREALLOCATED)
 	{
+		// Preallocated links: prefer base pool, fall back to dynamic if exhausted
 		if (tile->linksFreeList != DT_NULL_LINK)
 		{
 			newLink = tile->linksFreeList;
 			tile->linksFreeList = tile->links[newLink].next;
 		}
+		else
+		{
+			// Base pool exhausted - fall back to secondary storage
+			if (tile->dynamicFreeList == DT_NULL_LINK)
+			{
+				// No free secondary links, grow the array
+				dtLink emptyLink;
+				memset(&emptyLink, 0, sizeof(dtLink));
+				emptyLink.next = DT_NULL_LINK;
+
+				tile->dynamicFreeList = tile->dynamicLinks.size();
+				tile->dynamicLinks.push(emptyLink);
+			}
+
+			newLink = tile->dynamicFreeList;
+			tile->dynamicFreeList = tile->dynamicLinks[newLink].next;
+
+			// Offset into secondary range
+			newLink += tile->header->maxLinkCount;
+		}
 	}
-//@UE BEGIN: offmesh links will be added in dynamic array. 
-	// Both Point to Point (stock detour) and Segment to Segment (unreal)
+	// Off-mesh links use secondary storage to preserve base pool for internal/external links
 	else if (LinkAllocMode == CREATE_LINK_DYNAMIC_OFFMESH)
 	{
-		if (tile->dynamicFreeListO == DT_NULL_LINK)
+		// Allocate from secondary storage only (never consume base pool)
+		// This ensures internal/external links (PREALLOCATED) get base pool priority
+		if (tile->dynamicFreeList == DT_NULL_LINK)
 		{
+			// No free secondary links, grow the array
 			dtLink emptyLink;
 			memset(&emptyLink, 0, sizeof(dtLink));
 			emptyLink.next = DT_NULL_LINK;
 
-			tile->dynamicFreeListO = tile->dynamicLinksO.size();
-			tile->dynamicLinksO.push(emptyLink);		
+			tile->dynamicFreeList = tile->dynamicLinks.size();
+			tile->dynamicLinks.push(emptyLink);
 		}
 
-		newLink = tile->dynamicFreeListO;
-		tile->dynamicFreeListO = tile->dynamicLinksO[newLink].next;
+		newLink = tile->dynamicFreeList;
+		tile->dynamicFreeList = tile->dynamicLinks[newLink].next;
 
+		// Offset into secondary range
 		newLink += tile->header->maxLinkCount;
 	}
 #if WITH_NAVMESH_CLUSTER_LINKS
@@ -290,8 +313,8 @@ inline void freeLink(dtMeshTile* tile, unsigned int link)
 	else if (link < firstClusterLinkIdx)
 	{
 		const unsigned int linkIdx = link - tile->header->maxLinkCount;
-		tile->dynamicLinksO[linkIdx].next = tile->dynamicFreeListO;
-		tile->dynamicFreeListO = linkIdx;
+		tile->dynamicLinks[linkIdx].next = tile->dynamicFreeList;	//@UE
+		tile->dynamicFreeList = linkIdx;							//@UE
 	}
 #if WITH_NAVMESH_CLUSTER_LINKS
 	else
@@ -325,7 +348,7 @@ void dtFreeNavMesh(dtNavMesh* navmesh)
 //@UE BEGIN
 void dtFreeNavMeshTileRuntimeData(dtMeshTile* tile)
 {
-	tile->dynamicLinksO.~dtChunkArray();
+	tile->dynamicLinks.~dtChunkArray();	//@UE
 #if WITH_NAVMESH_CLUSTER_LINKS
 	tile->dynamicLinksC.~dtChunkArray();
 #endif // WITH_NAVMESH_CLUSTER_LINKS
@@ -1214,8 +1237,6 @@ void dtNavMesh::connectExtLinks(dtMeshTile* tile, dtMeshTile* target, int side,
 		dtPoly* poly = &tile->polys[i];
 
 		// Create new links.
-//		unsigned short m = DT_EXT_LINK | (unsigned short)side;
-		
 		const int nv = poly->vertCount;
 		for (int j = 0; j < nv; ++j)
 		{
@@ -1242,12 +1263,12 @@ void dtNavMesh::connectExtLinks(dtMeshTile* tile, dtMeshTile* target, int side,
 				unsigned int idx = allocLink(tile, CREATE_LINK_PREALLOCATED);
 				if (idx != DT_NULL_LINK)
 				{
-					dtLink* link = &tile->links[idx];
-					link->ref = NeiData.ref;
-					link->edge = (unsigned char)j;
-					link->side = (unsigned char)dir;
-					
-					link->next = poly->firstLink;
+					dtLink& link = getLink(tile, idx);	//@UE
+					link.ref = NeiData.ref;				//@UE
+					link.edge = (unsigned char)j;		//@UE
+					link.side = (unsigned char)dir;		//@UE
+
+					link.next = poly->firstLink;		//@UE
 					poly->firstLink = idx;
 
 					// Compress portal limits to a byte value.
@@ -1257,8 +1278,8 @@ void dtNavMesh::connectExtLinks(dtMeshTile* tile, dtMeshTile* target, int side,
 						dtReal tmax = (NeiData.max-va[2]) / (vb[2]-va[2]);
 						if (tmin > tmax)
 							dtSwap(tmin,tmax);
-						link->bmin = (unsigned char)(dtClamp(tmin, 0.0f, 1.0f)*255.0f);
-						link->bmax = (unsigned char)(dtClamp(tmax, 0.0f, 1.0f)*255.0f);
+						link.bmin = (unsigned char)(dtClamp(tmin, 0.0f, 1.0f)*255.0f);	//@UE
+						link.bmax = (unsigned char)(dtClamp(tmax, 0.0f, 1.0f)*255.0f);	//@UE
 					}
 					else if (dir == 2 || dir == 6)
 					{
@@ -1266,8 +1287,8 @@ void dtNavMesh::connectExtLinks(dtMeshTile* tile, dtMeshTile* target, int side,
 						dtReal tmax = (NeiData.max-va[0]) / (vb[0]-va[0]);
 						if (tmin > tmax)
 							dtSwap(tmin,tmax);
-						link->bmin = (unsigned char)(dtClamp(tmin, 0.0f, 1.0f)*255.0f);
-						link->bmax = (unsigned char)(dtClamp(tmax, 0.0f, 1.0f)*255.0f);
+						link.bmin = (unsigned char)(dtClamp(tmin, 0.0f, 1.0f)*255.0f);	//@UE
+						link.bmax = (unsigned char)(dtClamp(tmax, 0.0f, 1.0f)*255.0f);	//@UE
 					}
 				}
 
@@ -1451,18 +1472,18 @@ void dtNavMesh::connectIntLinks(dtMeshTile* tile)
 			unsigned int idx = allocLink(tile, CREATE_LINK_PREALLOCATED);
 			if (idx != DT_NULL_LINK)
 			{
-				dtLink* link = &tile->links[idx];
-				link->ref = base | (dtPolyRef)(poly->neis[j]-1);
-				link->edge = (unsigned char)j;
-//@UE BEGIN 
-				link->side = DT_CONNECTION_INTERNAL;
-//@UE END
-				link->bmin = link->bmax = 0;
+				dtLink& link = getLink(tile, idx);				//@UE
+				link.ref = base | (dtPolyRef)(poly->neis[j]-1); //@UE
+				link.edge = (unsigned char)j;					//@UE
+
+				link.side = DT_CONNECTION_INTERNAL;				//@UE
+
+				link.bmin = link.bmax = 0;						//@UE
 				// Add to linked list.
-				link->next = poly->firstLink;
+				link.next = poly->firstLink;					//@UE
 				poly->firstLink = idx;
 			}
-		}			
+		}
 	}
 }
 
@@ -2047,8 +2068,8 @@ dtStatus dtNavMesh::addTile(unsigned char* data, int dataSize, int flags,
 		tile->links[i].next = i+1;
 
 	// Initialize dynamic links array
-	tile->dynamicFreeListO = DT_NULL_LINK;
-	tile->dynamicLinksO.resize(0);
+	tile->dynamicFreeList = DT_NULL_LINK;	//@UE
+	tile->dynamicLinks.resize(0);			//@UE
 
 #if WITH_NAVMESH_CLUSTER_LINKS
 	tile->dynamicFreeListC = DT_NULL_LINK;
diff --git a/Engine/Source/Runtime/Navmesh/Private/Detour/DetourNavMeshBuilder.cpp b/Engine/Source/Runtime/Navmesh/Private/Detour/DetourNavMeshBuilder.cpp
--- a/Engine/Source/Runtime/Navmesh/Private/Detour/DetourNavMeshBuilder.cpp
+++ b/Engine/Source/Runtime/Navmesh/Private/Detour/DetourNavMeshBuilder.cpp
@@ -437,8 +437,23 @@ bool dtCreateNavMeshData(dtNavMeshCreateParams* params, unsigned char** outData,
 		}
 	}
 
-//@UE BEGIN: offmesh links will be added in dynamic array
-	const int maxLinkCount = edgeCount + portalCount*2;
+//@UE BEGIN: Link pool sizing
+	// Compute the size of the base link pool.
+	// Offmesh links and extra preallocated links are added in a dynamic array.
+	int maxLinkCount;
+	if (params->bMinimizeLinkPoolSize)
+	{
+		// Minimized link pool size.
+		// Navigation meshes are primarily triangles/quads (3-4 edges per polygon).
+		// Each edge may connect to one or more neighbors, requiring links for each connection.
+		// The x4 multiplier provides enough capacity for typical connectivity patterns and
+		// tiles exceeding this capacity fall back to secondary storage.
+		maxLinkCount = dtMax(params->polyCount * 4, 16);
+	}
+	else
+	{
+		maxLinkCount = edgeCount + portalCount*2;
+	}
 //@UE END
 	
 	// Find unique detail vertices.
diff --git a/Engine/Source/Runtime/Navmesh/Private/DetourTileCache/DetourTileCache.cpp b/Engine/Source/Runtime/Navmesh/Private/DetourTileCache/DetourTileCache.cpp
--- a/Engine/Source/Runtime/Navmesh/Private/DetourTileCache/DetourTileCache.cpp
+++ b/Engine/Source/Runtime/Navmesh/Private/DetourTileCache/DetourTileCache.cpp
@@ -752,6 +752,7 @@ dtStatus dtTileCache::buildNavMeshTile(const dtCompressedTileRef ref, dtNavMesh*
 	params.cs = m_params.cs;
 	params.ch = m_params.ch;
 	params.buildBvTree = false;
+	params.bMinimizeLinkPoolSize = m_params.bMinimizeLinkPoolSize;
 	dtVcopy(params.bmin, tile->header->bmin);
 	dtVcopy(params.bmax, tile->header->bmax);
 	//@UE BEGIN
diff --git a/Engine/Source/Runtime/Navmesh/Public/Detour/DetourAlloc.h b/Engine/Source/Runtime/Navmesh/Public/Detour/DetourAlloc.h
--- a/Engine/Source/Runtime/Navmesh/Public/Detour/DetourAlloc.h
+++ b/Engine/Source/Runtime/Navmesh/Public/Detour/DetourAlloc.h
@@ -41,7 +41,7 @@ enum dtAllocHint
 	DT_ALLOC_PERM_PROXIMITY_GRID,
 	DT_ALLOC_PERM_TILE_CACHE_LAYER,
 	DT_ALLOC_PERM_TILE_DATA,
-	DT_ALLOC_PERM_TILE_DYNLINK_OFFMESH,
+	DT_ALLOC_PERM_TILE_DYNLINK,
 	DT_ALLOC_PERM_TILE_DYNLINK_CLUSTER,
 	DT_ALLOC_PERM_TILES,
 	DT_ALLOC_PERM_TILE_LINK_BUILDER,
diff --git a/Engine/Source/Runtime/Navmesh/Public/Detour/DetourNavMesh.h b/Engine/Source/Runtime/Navmesh/Public/Detour/DetourNavMesh.h
--- a/Engine/Source/Runtime/Navmesh/Public/Detour/DetourNavMesh.h
+++ b/Engine/Source/Runtime/Navmesh/Public/Detour/DetourNavMesh.h
@@ -505,8 +505,8 @@ struct dtMeshTile
 	dtChunkArray<dtClusterLink, DT_ALLOC_PERM_TILE_DYNLINK_CLUSTER> dynamicLinksC; ///< Dynamic links array (indices starting from DT_CLINK_FIRST)
 #endif // WITH_NAVMESH_CLUSTER_LINKS
 
-	dtChunkArray<dtLink, DT_ALLOC_PERM_TILE_DYNLINK_OFFMESH> dynamicLinksO; ///< Dynamic links array (indices starting from dtMeshHeader::maxLinkCount)
-	unsigned int dynamicFreeListO;				///< Index of the next free dynamic link
+	dtChunkArray<dtLink, DT_ALLOC_PERM_TILE_DYNLINK> dynamicLinks; ///< Dynamic links array (indices starting from dtMeshHeader::maxLinkCount)	//@UE
+	unsigned int dynamicFreeList;				///< Index of the next free dynamic link //@UE
 
 #if WITH_NAVMESH_CLUSTER_LINKS
 	unsigned int dynamicFreeListC;				///< Index of the next free dynamic link
@@ -868,12 +868,12 @@ public:
 	/// Helper for accessing links
 	inline dtLink& getLink(dtMeshTile* tile, unsigned int linkIdx)
 	{
-		return (linkIdx < (unsigned int)tile->header->maxLinkCount) ? tile->links[linkIdx] : tile->dynamicLinksO[linkIdx - tile->header->maxLinkCount];
+		return (linkIdx < (unsigned int)tile->header->maxLinkCount) ? tile->links[linkIdx] : tile->dynamicLinks[linkIdx - tile->header->maxLinkCount]; //@UE
 	}
 
 	inline const dtLink& getLink(const dtMeshTile* tile, unsigned int linkIdx) const
 	{
-		return (linkIdx < (unsigned int)tile->header->maxLinkCount) ? tile->links[linkIdx] : tile->dynamicLinksO[linkIdx - tile->header->maxLinkCount];
+		return (linkIdx < (unsigned int)tile->header->maxLinkCount) ? tile->links[linkIdx] : tile->dynamicLinks[linkIdx - tile->header->maxLinkCount]; //@UE
 	}
 
 #if WITH_NAVMESH_CLUSTER_LINKS
diff --git a/Engine/Source/Runtime/Navmesh/Public/Detour/DetourNavMeshBuilder.h b/Engine/Source/Runtime/Navmesh/Public/Detour/DetourNavMeshBuilder.h
--- a/Engine/Source/Runtime/Navmesh/Public/Detour/DetourNavMeshBuilder.h
+++ b/Engine/Source/Runtime/Navmesh/Public/Detour/DetourNavMeshBuilder.h
@@ -149,6 +149,11 @@ struct dtNavMeshCreateParams
 	/// @note The BVTree is not normally needed for layered navigation meshes.
 	bool buildBvTree;
 
+	//@UE BEGIN
+	/// If true, minimize base link pool size else use previous larger base pool.
+	bool bMinimizeLinkPoolSize;
+	//@UE END
+
 	/// @}
 };
 
diff --git a/Engine/Source/Runtime/Navmesh/Public/DetourTileCache/DetourTileCache.h b/Engine/Source/Runtime/Navmesh/Public/DetourTileCache/DetourTileCache.h
--- a/Engine/Source/Runtime/Navmesh/Public/DetourTileCache/DetourTileCache.h
+++ b/Engine/Source/Runtime/Navmesh/Public/DetourTileCache/DetourTileCache.h
@@ -92,6 +92,7 @@ struct dtTileCacheParams
 	int mergeRegionArea;
 	int regionChunkSize;
 	int regionPartitioning;
+	bool bMinimizeLinkPoolSize;
 //@UE END
 };
 
