[Navigation][Avoidance] Add configurable side bias threshold and right-side passing preference to Detour crowd avoidance

Exposes the hardcoded 0.01 side selection threshold in dtObstacleAvoidanceQuery::prepare() as a
configurable parameter (sideBiasThreshold) and adds a new right-side passing penalty (weightRightBias)
to processSample(). Both reduce left-right oscillation when agents approach each other head-on.

#Synced-CL 52534494

diff --git a/Engine/Source/Runtime/AIModule/Classes/Navigation/CrowdManager.h b/Engine/Source/Runtime/AIModule/Classes/Navigation/CrowdManager.h
--- a/Engine/Source/Runtime/AIModule/Classes/Navigation/CrowdManager.h
+++ b/Engine/Source/Runtime/AIModule/Classes/Navigation/CrowdManager.h
@@ -88,6 +88,14 @@ struct FCrowdAvoidanceConfig
 	UPROPERTY(EditAnywhere, Category=Crowd)
 	float ImpactTimeRange;
 
+	/** Threshold for side selection, larger values widen the dead zone for head-on encounters */
+	UPROPERTY(EditAnywhere, Category=Crowd, meta=(ClampMin="0.001", ClampMax="1.0"))
+	float SideBiasThreshold;
+
+	/** Weight for preferring right-side passing, reduces oscillation (0 = disabled) */
+	UPROPERTY(EditAnywhere, Category=Crowd, meta=(ClampMin="0.0"))
+	float RightSideBias;
+
 	// index in SamplingPatterns array or 0xff for adaptive sampling
 	UPROPERTY(EditAnywhere, Category=Crowd)
 	uint8 CustomPatternIdx;
@@ -107,6 +115,7 @@ struct FCrowdAvoidanceConfig
 	FCrowdAvoidanceConfig() :
 		VelocityBias(0.4f), DesiredVelocityWeight(2.0f), CurrentVelocityWeight(0.75f),
 		SideBiasWeight(0.75f), ImpactTimeWeight(2.5f), ImpactTimeRange(2.5f),
+		SideBiasThreshold(0.01f), RightSideBias(0.0f),
 		CustomPatternIdx(0xff), AdaptiveDivisions(7), AdaptiveRings(2), AdaptiveDepth(5)
 	{}
 };
diff --git a/Engine/Source/Runtime/AIModule/Private/Navigation/CrowdManager.cpp b/Engine/Source/Runtime/AIModule/Private/Navigation/CrowdManager.cpp
--- a/Engine/Source/Runtime/AIModule/Private/Navigation/CrowdManager.cpp
+++ b/Engine/Source/Runtime/AIModule/Private/Navigation/CrowdManager.cpp
@@ -1410,6 +1410,8 @@ void UCrowdManager::UpdateAvoidanceConfig()
 		params.weightSide = ConfigInfo.SideBiasWeight;
 		params.weightToi = ConfigInfo.ImpactTimeWeight;
 		params.horizTime = ConfigInfo.ImpactTimeRange;
+		params.sideBiasThreshold = ConfigInfo.SideBiasThreshold;
+		params.weightRightBias = ConfigInfo.RightSideBias;
 		params.patternIdx = ConfigInfo.CustomPatternIdx;
 		params.adaptiveDivs = ConfigInfo.AdaptiveDivisions;
 		params.adaptiveRings = ConfigInfo.AdaptiveRings;
diff --git a/Engine/Source/Runtime/Navmesh/Private/DetourCrowd/DetourObstacleAvoidance.cpp b/Engine/Source/Runtime/Navmesh/Private/DetourCrowd/DetourObstacleAvoidance.cpp
--- a/Engine/Source/Runtime/Navmesh/Private/DetourCrowd/DetourObstacleAvoidance.cpp
+++ b/Engine/Source/Runtime/Navmesh/Private/DetourCrowd/DetourObstacleAvoidance.cpp
@@ -297,7 +297,7 @@ void dtObstacleAvoidanceQuery::prepare(const dtReal* pos, const dtReal* dvel)
 		dtVsub(dv, cir->dvel, dvel);
 		
 		const dtReal a = dtTriArea2D(orig, cir->dp,dv);
-		if (a < 0.01f)
+		if (a < m_params.sideBiasThreshold) //@UE configurable threshold (was 0.01f)
 		{
 			cir->np[0] = -cir->dp[2];
 			cir->np[2] = cir->dp[0];
@@ -328,6 +328,7 @@ dtReal dtObstacleAvoidanceQuery::processSample(const dtReal* vcand, const dtReal
 	// Find min time of impact and exit amongst all obstacles.
 	dtReal tmin = m_params.horizTime;
 	dtReal side = 0;
+	dtReal rightBias = 0; //@UE right-side passing preference
 	int nside = 0;
 	
 	for (int i = 0; i < m_ncircles; ++i)
@@ -343,6 +344,17 @@ dtReal dtObstacleAvoidanceQuery::processSample(const dtReal* vcand, const dtReal
 		// Side
 		side += dtClamp(dtMin(dtVdot2D(cir->dp,vab)*0.5f+0.5f, dtVdot2D(cir->np,vab)*2), 0.0f, 1.0f);
 		nside++;
+
+		//@UE BEGIN - Right-side passing bias: penalize left-side passing to reduce head-on oscillation.
+		// Use a fixed right-side normal (-dp[2], dp[0]) independent of the np flip in prepare().
+		// Both agents in a head-on pair compute opposite dp vectors, so this consistently
+		// makes them agree to pass on their respective right sides.
+		if (m_params.weightRightBias > 0.0f)
+		{
+			const dtReal rightDot = -cir->dp[2] * vab[0] + cir->dp[0] * vab[2];
+			rightBias += dtClamp(rightDot * 0.5f + 0.5f, 0.0f, 1.0f);
+		}
+		//@UE END
 		
 		dtReal htmin = 0, htmax = 0;
 		if (!sweepCircleCircle(pos,rad, vab, cir->p,cir->rad, htmin, htmax))
@@ -409,14 +421,18 @@ dtReal dtObstacleAvoidanceQuery::processSample(const dtReal* vcand, const dtReal
 	
 	// Normalize side bias, to prevent it dominating too much.
 	if (nside)
+	{
 		side /= nside;
-	
+		rightBias /= nside; //@UE
+	}
+
 	const dtReal vpen = m_params.weightDesVel * (dtVdist2D(vcand, dvel) * m_invVmax);
 	const dtReal vcpen = m_params.weightCurVel * (dtVdist2D(vcand, vel) * m_invVmax);
 	const dtReal spen = m_params.weightSide * side;
 	const dtReal tpen = m_params.weightToi * (1.0f/(0.1f+tmin*m_invHorizTime));
+	const dtReal rpen = m_params.weightRightBias * rightBias; //@UE right-side passing penalty
 	
-	const dtReal penalty = vpen + vcpen + spen + tpen;
+	const dtReal penalty = vpen + vcpen + spen + tpen + rpen; //@UE added rpen
 	
 	// Store different penalties for debug viewing
 	if (debug)
@@ -458,9 +474,9 @@ int dtObstacleAvoidanceQuery::sampleVelocityCustom(const dtReal* pos, const dtRe
 												   const dtObstacleAvoidanceParams* params,
 											  	   dtObstacleAvoidanceDebugData* debug)
 {
+	memcpy(&m_params, params, sizeof(dtObstacleAvoidanceParams)); //@UE moved before prepare() so it can access m_params
 	prepare(pos, dvel);
-	
-	memcpy(&m_params, params, sizeof(dtObstacleAvoidanceParams));
+
 	m_invHorizTime = 1.0f / m_params.horizTime;
 	m_vmax = vmax;
 	m_invVmax = 1.0f / vmax;
@@ -523,9 +539,9 @@ int dtObstacleAvoidanceQuery::sampleVelocityAdaptive(const dtReal* pos, const dt
 													 const dtObstacleAvoidanceParams* params,
 													 dtObstacleAvoidanceDebugData* debug)
 {
+	memcpy(&m_params, params, sizeof(dtObstacleAvoidanceParams)); //@UE moved before prepare() so it can access m_params
 	prepare(pos, dvel);
-	
-	memcpy(&m_params, params, sizeof(dtObstacleAvoidanceParams));
+
 	m_invHorizTime = 1.0f / m_params.horizTime;
 	m_vmax = vmax * vmult;
 	m_invVmax = 1.0f / (vmax * vmult);
diff --git a/Engine/Source/Runtime/Navmesh/Public/DetourCrowd/DetourObstacleAvoidance.h b/Engine/Source/Runtime/Navmesh/Public/DetourCrowd/DetourObstacleAvoidance.h
--- a/Engine/Source/Runtime/Navmesh/Public/DetourCrowd/DetourObstacleAvoidance.h
+++ b/Engine/Source/Runtime/Navmesh/Public/DetourCrowd/DetourObstacleAvoidance.h
@@ -91,6 +91,8 @@ struct dtObstacleAvoidanceParams
 	dtReal weightSide;
 	dtReal weightToi;
 	dtReal horizTime;
+	dtReal sideBiasThreshold;	///< Threshold for side selection in prepare(), widens the dead zone for head-on encounters //@UE
+	dtReal weightRightBias;		///< Weight for preferring right-side passing (reduces head-on oscillation) //@UE
 	unsigned char patternIdx;	///< [UE] index of custom sampling pattern or 0xff for adaptive
 	unsigned char adaptiveDivs;	///< adaptive
 	unsigned char adaptiveRings;	///< adaptive
